[Mass]
Move MassSignals module from MassGameplay plugin to Engine.

diff --git a/Engine/Plugins/Runtime/MassGameplay/MassGameplay.uplugin b/Engine/Plugins/Runtime/MassGameplay/MassGameplay.uplugin
--- a/Engine/Plugins/Runtime/MassGameplay/MassGameplay.uplugin
+++ b/Engine/Plugins/Runtime/MassGameplay/MassGameplay.uplugin
@@ -28,10 +28,6 @@
 			"Name": "MassEQS",
 			"Type": "Runtime"
 		},
-		{
-			"Name": "MassSignals",
-			"Type": "Runtime"
-		},
 		{
 			"Name": "MassSpawner",
 			"Type": "Runtime"
diff --git a/Engine/Plugins/Runtime/MassGameplay/Source/MassSignals/MassSignals.Build.cs b/Engine/Plugins/Runtime/MassGameplay/Source/MassSignals/MassSignals.Build.cs
deleted file mode 100644
diff --git a/Engine/Plugins/Runtime/MassGameplay/Source/MassSignals/Private/MassSignalProcessorBase.cpp b/Engine/Plugins/Runtime/MassGameplay/Source/MassSignals/Private/MassSignalProcessorBase.cpp
deleted file mode 100644
diff --git a/Engine/Plugins/Runtime/MassGameplay/Source/MassSignals/Private/MassSignalSubsystem.cpp b/Engine/Plugins/Runtime/MassGameplay/Source/MassSignals/Private/MassSignalSubsystem.cpp
deleted file mode 100644
diff --git a/Engine/Plugins/Runtime/MassGameplay/Source/MassSignals/Private/MassSignalTypes.cpp b/Engine/Plugins/Runtime/MassGameplay/Source/MassSignals/Private/MassSignalTypes.cpp
deleted file mode 100644
diff --git a/Engine/Plugins/Runtime/MassGameplay/Source/MassSignals/Private/MassSignalsModule.cpp b/Engine/Plugins/Runtime/MassGameplay/Source/MassSignals/Private/MassSignalsModule.cpp
deleted file mode 100644
diff --git a/Engine/Plugins/Runtime/MassGameplay/Source/MassSignals/Public/IMassSignalsModule.h b/Engine/Plugins/Runtime/MassGameplay/Source/MassSignals/Public/IMassSignalsModule.h
deleted file mode 100644
diff --git a/Engine/Plugins/Runtime/MassGameplay/Source/MassSignals/Public/MassSignalProcessorBase.h b/Engine/Plugins/Runtime/MassGameplay/Source/MassSignals/Public/MassSignalProcessorBase.h
deleted file mode 100644
diff --git a/Engine/Plugins/Runtime/MassGameplay/Source/MassSignals/Public/MassSignalSubsystem.h b/Engine/Plugins/Runtime/MassGameplay/Source/MassSignals/Public/MassSignalSubsystem.h
deleted file mode 100644
diff --git a/Engine/Plugins/Runtime/MassGameplay/Source/MassSignals/Public/MassSignalTypes.h b/Engine/Plugins/Runtime/MassGameplay/Source/MassSignals/Public/MassSignalTypes.h
deleted file mode 100644
diff --git a/Engine/Source/Runtime/Engine/Engine.Build.cs b/Engine/Source/Runtime/Engine/Engine.Build.cs
--- a/Engine/Source/Runtime/Engine/Engine.Build.cs
+++ b/Engine/Source/Runtime/Engine/Engine.Build.cs
@@ -394,6 +394,7 @@ public class Engine : ModuleRules
 		}
 
 		DynamicallyLoadedModuleNames.Add("MassEntity");
+		DynamicallyLoadedModuleNames.Add("MassSignals");
 
 		if (Target.bBuildEditor == true)
 		{
diff --git a/Engine/Source/Runtime/Mass/MassSignals/MassSignals.Build.cs b/Engine/Source/Runtime/Mass/MassSignals/MassSignals.Build.cs
new file mode 100644
--- /dev/null
+++ b/Engine/Source/Runtime/Mass/MassSignals/MassSignals.Build.cs
@@ -0,0 +1,29 @@
+// Copyright Epic Games, Inc. All Rights Reserved.
+
+namespace UnrealBuildTool.Rules
+{
+	public class MassSignals : ModuleRules
+	{
+		public MassSignals(ReadOnlyTargetRules Target) : base(Target)
+		{
+			PCHUsage = ModuleRules.PCHUsageMode.UseExplicitOrSharedPCHs;
+
+			CppCompileWarningSettings.UnsafeTypeCastWarningLevel = WarningLevel.Warning;
+
+			PublicDependencyModuleNames.AddRange(
+				new string[] {
+					"Core",
+					"CoreUObject",
+					"Engine",
+					"MassEntity",
+				}
+			);
+
+			if (Target.bBuildEditor == true)
+			{
+				PrivateDependencyModuleNames.Add("EditorFramework");
+				PrivateDependencyModuleNames.Add("UnrealEd");
+			}
+		}
+	}
+}
diff --git a/Engine/Source/Runtime/Mass/MassSignals/Private/MassSignalProcessorBase.cpp b/Engine/Source/Runtime/Mass/MassSignals/Private/MassSignalProcessorBase.cpp
new file mode 100644
--- /dev/null
+++ b/Engine/Source/Runtime/Mass/MassSignals/Private/MassSignalProcessorBase.cpp
@@ -0,0 +1,158 @@
+﻿// Copyright Epic Games, Inc. All Rights Reserved.
+
+#include "MassSignalProcessorBase.h"
+#include "MassSignalSubsystem.h"
+#include "MassArchetypeTypes.h"
+#include "MassExecutionContext.h"
+#include "Engine/World.h"
+#include "Misc/ScopeLock.h"
+
+#include UE_INLINE_GENERATED_CPP_BY_NAME(MassSignalProcessorBase)
+
+
+UMassSignalProcessorBase::UMassSignalProcessorBase(const FObjectInitializer& ObjectInitializer)
+	: EntityQuery(*this)
+{
+	ExecutionFlags = (int32)EProcessorExecutionFlags::AllNetModes;
+}
+
+void UMassSignalProcessorBase::BeginDestroy()
+{
+	if (UMassSignalSubsystem* SignalSubsystem = UWorld::GetSubsystem<UMassSignalSubsystem>(GetWorld()))
+	{
+		for (const FName& SignalName : RegisteredSignals)
+		{
+			SignalSubsystem->GetSignalDelegateByName(SignalName).RemoveAll(this);
+		}
+	}
+
+	Super::BeginDestroy();
+}
+
+void UMassSignalProcessorBase::SubscribeToSignal(UMassSignalSubsystem& SignalSubsystem, const FName SignalName)
+{
+	check(!RegisteredSignals.Contains(SignalName));
+	RegisteredSignals.Add(SignalName);
+	SignalSubsystem.GetSignalDelegateByName(SignalName).AddUObject(this, &UMassSignalProcessorBase::OnSignalReceived);
+}
+
+void UMassSignalProcessorBase::Execute(FMassEntityManager& EntityManager, FMassExecutionContext& Context)
+{
+	QUICK_SCOPE_CYCLE_COUNTER(SignalEntities);
+
+	const int32 ProcessingFrameBufferIndex = CurrentFrameBufferIndex;
+	{
+		// we only need to lock the part where we change the current buffer index. Once that's done the incoming signals will end up 
+		// in the other buffer
+		UE::TRWScopeLock Lock(ReceivedSignalLock, SLT_Write);
+		CurrentFrameBufferIndex = (CurrentFrameBufferIndex + 1) % BuffersCount;
+	}
+	
+	FFrameReceivedSignals& ProcessingFrameBuffer = FrameReceivedSignals[ProcessingFrameBufferIndex];
+	TArray<FEntitySignalRange>& ReceivedSignalRanges = ProcessingFrameBuffer.ReceivedSignalRanges;
+	TArray<FMassEntityHandle>& SignaledEntities = ProcessingFrameBuffer.SignaledEntities;
+
+	if (ReceivedSignalRanges.IsEmpty())
+	{
+		return;
+	}
+
+	TArray<FMassArchetypeHandle> ValidArchetypes;
+	GetArchetypesMatchingOwnedQueries(EntityManager, ValidArchetypes);
+
+	if (ValidArchetypes.Num() > 0)
+	{
+		// EntitySet stores unique array of entities per specified archetype.
+		// FMassArchetypeEntityCollection expects an array of entities, a set is used to detect unique ones.
+		struct FEntitySet
+		{
+			void Reset()
+			{
+				Entities.Reset();
+			}
+
+			FMassArchetypeHandle Archetype;
+			TArray<FMassEntityHandle> Entities;
+		};
+		TArray<FEntitySet> EntitySets;
+
+		for (const FMassArchetypeHandle& Archetype : ValidArchetypes)
+		{
+			FEntitySet& Set = EntitySets.AddDefaulted_GetRef();
+			Set.Archetype = Archetype;
+		}
+
+		// SignalNameLookup has limit of how many signals it can handle at once, we'll do passes until all signals are processed.
+		int32 SignalsToProcess = ReceivedSignalRanges.Num();
+		while(SignalsToProcess > 0)
+		{
+			SignalNameLookup.Reset();
+
+			// Convert signals with entity ids into arrays of entities per archetype.
+			for (FEntitySignalRange& Range : ReceivedSignalRanges)
+			{
+				if (Range.bProcessed)
+				{
+					continue;
+				}
+				// Get bitflag for the signal name
+				const uint64 SignalFlag = SignalNameLookup.GetOrAddSignalName(Range.SignalName);
+				if (SignalFlag == 0)
+				{
+					// Will process that signal in a second iteration
+					continue;
+				}
+
+				// Entities for this signal
+				TArrayView<FMassEntityHandle> Entities(&SignaledEntities[Range.Begin], Range.End - Range.Begin);
+				FEntitySet* PrevSet = &EntitySets[0];
+				for (const FMassEntityHandle Entity : Entities)
+				{
+					// Add to set of supported archetypes. Dont process if we don't care about the type.
+					FMassArchetypeHandle Archetype = EntityManager.GetArchetypeForEntity(Entity);
+					FEntitySet* Set = PrevSet->Archetype == Archetype ? PrevSet : EntitySets.FindByPredicate([&Archetype](const FEntitySet& Set) { return Archetype == Set.Archetype; });
+					if (Set != nullptr)
+					{
+						// We don't care about duplicates here, the FMassArchetypeEntityCollection creation below will handle it
+						Set->Entities.Add(Entity);
+						SignalNameLookup.AddSignalToEntity(Entity, SignalFlag);
+						PrevSet = Set;
+					}
+				}
+
+				Range.bProcessed = true;
+				SignalsToProcess--;
+			}
+
+			// Execute per archetype.
+			for (FEntitySet& Set : EntitySets)
+			{
+				if (Set.Entities.Num() > 0)
+				{
+					Context.SetEntityCollection(FMassArchetypeEntityCollection(Set.Archetype, Set.Entities, FMassArchetypeEntityCollection::FoldDuplicates));
+					SignalEntities(EntityManager, Context, SignalNameLookup);
+					Context.ClearEntityCollection();
+				}
+				Set.Reset();
+			}
+		}
+	}
+
+	ReceivedSignalRanges.Reset();
+	SignaledEntities.Reset();
+}
+
+void UMassSignalProcessorBase::OnSignalReceived(FName SignalName, TConstArrayView<FMassEntityHandle> Entities)
+{
+	FEntitySignalRange Range;
+	Range.SignalName = SignalName;
+
+	UE::TRWScopeLock Lock(ReceivedSignalLock, SLT_Write);
+
+	FFrameReceivedSignals& CurrentFrameBuffer = FrameReceivedSignals[CurrentFrameBufferIndex];
+
+	Range.Begin = CurrentFrameBuffer.SignaledEntities.Num();
+	CurrentFrameBuffer.SignaledEntities.Append(Entities.GetData(), Entities.Num());
+	Range.End = CurrentFrameBuffer.SignaledEntities.Num();
+	CurrentFrameBuffer.ReceivedSignalRanges.Add(MoveTemp(Range));
+}
diff --git a/Engine/Source/Runtime/Mass/MassSignals/Private/MassSignalSubsystem.cpp b/Engine/Source/Runtime/Mass/MassSignals/Private/MassSignalSubsystem.cpp
new file mode 100644
--- /dev/null
+++ b/Engine/Source/Runtime/Mass/MassSignals/Private/MassSignalSubsystem.cpp
@@ -0,0 +1,140 @@
+// Copyright Epic Games, Inc. All Rights Reserved.
+
+#include "MassSignalSubsystem.h"
+#include "MassEntityManager.h"
+#include "MassExecutionContext.h"
+#include "MassCommandBuffer.h"
+#include "MassSignalTypes.h"
+#include "VisualLogger/VisualLogger.h"
+#include "ProfilingDebugging/CsvProfiler.h"
+#include "Engine/World.h"
+
+#include UE_INLINE_GENERATED_CPP_BY_NAME(MassSignalSubsystem)
+
+CSV_DEFINE_CATEGORY(MassSignalsCounters, true);
+
+void UMassSignalSubsystem::Initialize(FSubsystemCollectionBase& Collection)
+{
+	Super::Initialize(Collection);
+	CachedWorld = &GetWorldRef();
+	checkf(CachedWorld, TEXT("UMassSignalSubsystem instances are expected to always be tied to a valid UWorld instance"));
+
+	OverrideSubsystemTraits<UMassSignalSubsystem>(Collection);
+}
+
+void UMassSignalSubsystem::Deinitialize()
+{
+	CachedWorld = nullptr;
+	Super::Deinitialize();
+}
+
+void UMassSignalSubsystem::Tick(float DeltaTime)
+{
+	CA_ASSUME(CachedWorld);
+	const double CurrentTime = CachedWorld->GetTimeSeconds();
+
+	UE_MT_SCOPED_WRITE_ACCESS(DelayedSignalsAccessDetector);
+
+	for (int i = 0; i < DelayedSignals.Num();)
+	{
+		FDelayedSignal& DelayedSignal = DelayedSignals[i];
+		if (DelayedSignal.TargetTimestamp <= CurrentTime)
+		{
+			SignalEntities(DelayedSignal.SignalName, MakeArrayView(DelayedSignal.Entities));
+			DelayedSignals.RemoveAtSwap(i, EAllowShrinking::No);
+		}
+		else
+		{
+			i++;
+		}
+	}
+}
+
+TStatId UMassSignalSubsystem::GetStatId() const
+{
+	RETURN_QUICK_DECLARE_CYCLE_STAT(UMassSignalSubsystem, STATGROUP_Tickables);
+}
+
+void UMassSignalSubsystem::SignalEntity(FName SignalName, const FMassEntityHandle Entity)
+{
+	checkf(Entity.IsSet(), TEXT("Expecting a valid entity to signal"));
+	SignalEntities(SignalName, MakeArrayView(&Entity, 1));
+}
+
+void UMassSignalSubsystem::SignalEntities(FName SignalName, TConstArrayView<FMassEntityHandle> Entities)
+{
+	checkf(Entities.Num() > 0, TEXT("Expecting entities to signal"));
+	const UE::MassSignal::FSignalDelegate& SignalDelegate = GetSignalDelegateByName(SignalName);
+	SignalDelegate.Broadcast(SignalName, Entities);
+
+#if CSV_PROFILER_STATS
+	FCsvProfiler::RecordCustomStat(*SignalName.ToString(), CSV_CATEGORY_INDEX(MassSignalsCounters), Entities.Num(), ECsvCustomStatOp::Accumulate);
+#endif
+
+	UE_CVLOG(Entities.Num() == 1, this, LogMassSignals, Log, TEXT("Raising signal [%s] to entity [%s]"), *SignalName.ToString(), *Entities[0].DebugGetDescription());
+	UE_CVLOG(Entities.Num() > 1, this, LogMassSignals, Log, TEXT("Raising signal [%s] to %d entities"), *SignalName.ToString(), Entities.Num());
+}
+
+void UMassSignalSubsystem::DelaySignalEntity(FName SignalName, const FMassEntityHandle Entity, const float DelayInSeconds)
+{
+	checkf(Entity.IsSet(), TEXT("Expecting a valid entity to signal"));
+	DelaySignalEntities(SignalName, MakeArrayView(&Entity, 1), DelayInSeconds);
+}
+
+void UMassSignalSubsystem::DelaySignalEntities(FName SignalName, TConstArrayView<FMassEntityHandle> Entities, const float DelayInSeconds)
+{
+	// If you hit this ensure
+	// - With another thread trying to delay signal then you can use DelaySignalEntityDeferred/DelaySignalEntitiesDeferred
+	//   if you have access to a FMassExecutionContext.
+	// - With the game thread executing UMassSignalSubsystem::Tick then you need to reorganize your tasks to prevent senders from executing
+	//   at the same time as the subsystem tick.
+	UE_MT_SCOPED_WRITE_ACCESS(DelayedSignalsAccessDetector);
+
+	FDelayedSignal& DelayedSignal = DelayedSignals.Emplace_GetRef();
+	DelayedSignal.SignalName = SignalName;
+	DelayedSignal.Entities = Entities;
+
+	check(CachedWorld);
+	DelayedSignal.TargetTimestamp = CachedWorld->GetTimeSeconds() + DelayInSeconds;
+
+	UE_CVLOG(Entities.Num() == 1, this, LogMassSignals, Log, TEXT("Delay signal [%s] to entity [%s] in %.2f"), *SignalName.ToString(), *Entities[0].DebugGetDescription(), DelayInSeconds);
+	UE_CVLOG(Entities.Num() > 1,this, LogMassSignals, Log, TEXT("Delay signal [%s] to %d entities in %.2f"), *SignalName.ToString(), Entities.Num(), DelayInSeconds);
+}
+
+void UMassSignalSubsystem::SignalEntityDeferred(FMassExecutionContext& Context, FName SignalName, const FMassEntityHandle Entity)
+{
+	checkf(Entity.IsSet(), TEXT("Expecting a valid entity to signal"));
+	SignalEntitiesDeferred(Context, SignalName, MakeArrayView(&Entity, 1));
+}
+
+void UMassSignalSubsystem::SignalEntitiesDeferred(FMassExecutionContext& Context, FName SignalName, TConstArrayView<FMassEntityHandle> Entities)
+{
+	checkf(Entities.Num() > 0, TEXT("Expecting entities to signal"));
+	Context.Defer().PushCommand<FMassDeferredSetCommand>([SignalName, InEntities = TArray<FMassEntityHandle>(Entities)](const FMassEntityManager& InOutEntityManager)
+	{
+		UMassSignalSubsystem* SignalSubsystem = UWorld::GetSubsystem<UMassSignalSubsystem>(InOutEntityManager.GetWorld());
+		SignalSubsystem->SignalEntities(SignalName, InEntities);
+	});
+
+	UE_CVLOG(Entities.Num() == 1, this, LogMassSignals, Log, TEXT("Raising deferred signal [%s] to entity [%s]"), *SignalName.ToString(), *Entities[0].DebugGetDescription());
+	UE_CVLOG(Entities.Num() > 1, this, LogMassSignals, Log, TEXT("Raising deferred signal [%s] to %d entities"), *SignalName.ToString(), Entities.Num());
+}
+
+void UMassSignalSubsystem::DelaySignalEntityDeferred(FMassExecutionContext& Context, FName SignalName, const FMassEntityHandle Entity, const float DelayInSeconds)
+{
+	checkf(Entity.IsSet(), TEXT("Expecting a valid entity to signal"));
+	DelaySignalEntitiesDeferred(Context, SignalName, MakeArrayView(&Entity, 1), DelayInSeconds);
+}
+
+void UMassSignalSubsystem::DelaySignalEntitiesDeferred(FMassExecutionContext& Context, FName SignalName, TConstArrayView<FMassEntityHandle> Entities, const float DelayInSeconds)
+{
+	checkf(Entities.Num() > 0, TEXT("Expecting entities to signal"));
+	Context.Defer().PushCommand<FMassDeferredSetCommand>([SignalName, InEntities = TArray<FMassEntityHandle>(Entities), DelayInSeconds](const FMassEntityManager& InOutEntityManager)
+	{
+		UMassSignalSubsystem* SignalSubsystem = UWorld::GetSubsystem<UMassSignalSubsystem>(InOutEntityManager.GetWorld());
+		SignalSubsystem->DelaySignalEntities(SignalName, InEntities, DelayInSeconds);
+	});
+
+	UE_CVLOG(Entities.Num() == 1, this, LogMassSignals, Log, TEXT("Delay deferred signal [%s] to entity [%s] in %.2f"), *SignalName.ToString(), *Entities[0].DebugGetDescription(), DelayInSeconds);
+	UE_CVLOG(Entities.Num() > 1,this, LogMassSignals, Log, TEXT("Delay deferred signal [%s] to %d entities in %.2f"), *SignalName.ToString(), Entities.Num(), DelayInSeconds);
+}
diff --git a/Engine/Source/Runtime/Mass/MassSignals/Private/MassSignalTypes.cpp b/Engine/Source/Runtime/Mass/MassSignals/Private/MassSignalTypes.cpp
new file mode 100644
--- /dev/null
+++ b/Engine/Source/Runtime/Mass/MassSignals/Private/MassSignalTypes.cpp
@@ -0,0 +1,48 @@
+// Copyright Epic Games, Inc. All Rights Reserved.
+
+#include "MassSignalTypes.h"
+#include "MassEntityManager.h"
+
+uint64 FMassSignalNameLookup::GetOrAddSignalName(const FName SignalName)
+{
+	int32 Index = SignalNames.Find(SignalName);
+	if (Index == INDEX_NONE)
+	{
+		if (SignalNames.Num() >= MaxSignalNames)
+		{
+			return 0;
+		}
+		Index = SignalNames.Add(SignalName);
+	}
+	return 1ULL << Index;
+}
+
+void FMassSignalNameLookup::AddSignalToEntity(const FMassEntityHandle Entity, const uint64 SignalFlag)
+{
+	uint64& Signals = EntitySignals.FindOrAdd(Entity, 0);
+	Signals |= SignalFlag;
+}
+
+void FMassSignalNameLookup::GetSignalsForEntity(const FMassEntityHandle Entity, TArray<FName>& OutSignals) const
+{
+	OutSignals.Reset();
+	if (const uint64* Signals = EntitySignals.Find(Entity))
+	{
+		const uint64 Start = FMath::CountTrailingZeros64(*Signals);
+		const uint64 End = FMath::CountLeadingZeros64(*Signals);
+		for (uint64 i = Start; i < 64 - End; i++)
+		{
+			const uint64 SignalFlag = 1ULL << i;
+			if ((*Signals) & SignalFlag)
+			{
+				OutSignals.Add(SignalNames[(uint32)i]);
+			}
+		}
+	}
+}
+
+void FMassSignalNameLookup::Reset()
+{
+	SignalNames.Reset();
+	EntitySignals.Reset();
+}
\ No newline at end of file
diff --git a/Engine/Source/Runtime/Mass/MassSignals/Private/MassSignalsModule.cpp b/Engine/Source/Runtime/Mass/MassSignals/Private/MassSignalsModule.cpp
new file mode 100644
--- /dev/null
+++ b/Engine/Source/Runtime/Mass/MassSignals/Private/MassSignalsModule.cpp
@@ -0,0 +1,34 @@
+// Copyright Epic Games, Inc. All Rights Reserved.
+
+#include "CoreMinimal.h"
+#include "MassSignalTypes.h"
+#include "IMassSignalsModule.h"
+
+DEFINE_LOG_CATEGORY(LogMassSignals)
+
+
+class FMassSignalsModule : public IMassSignalsModule
+{
+	/** IModuleInterface implementation */
+	virtual void StartupModule() override;
+	virtual void ShutdownModule() override;
+};
+
+IMPLEMENT_MODULE(FMassSignalsModule, MassSignals)
+
+
+
+void FMassSignalsModule::StartupModule()
+{
+	// This code will execute after your module is loaded into memory (but after global variables are initialized, of course.)
+}
+
+
+void FMassSignalsModule::ShutdownModule()
+{
+	// This function may be called during shutdown to clean up your module.  For modules that support dynamic reloading,
+	// we call this function before unloading the module.
+}
+
+
+
diff --git a/Engine/Source/Runtime/Mass/MassSignals/Public/IMassSignalsModule.h b/Engine/Source/Runtime/Mass/MassSignals/Public/IMassSignalsModule.h
new file mode 100644
--- /dev/null
+++ b/Engine/Source/Runtime/Mass/MassSignals/Public/IMassSignalsModule.h
@@ -0,0 +1,40 @@
+// Copyright Epic Games, Inc. All Rights Reserved.
+
+#pragma once
+
+#include "CoreMinimal.h"
+#include "Modules/ModuleInterface.h"
+#include "Modules/ModuleManager.h"
+
+
+/**
+ * The public interface to this module.  In most cases, this interface is only public to sibling modules 
+ * within this plugin.
+ */
+class IMassSignalsModule : public IModuleInterface
+{
+
+public:
+
+	/**
+	 * Singleton-like access to this module's interface.  This is just for convenience!
+	 * Beware of calling this during the shutdown phase, though.  Your module might have been unloaded already.
+	 *
+	 * @return Returns singleton instance, loading the module on demand if needed
+	 */
+	static inline IMassSignalsModule& Get()
+	{
+		return FModuleManager::LoadModuleChecked<IMassSignalsModule>( "MassSignals" );
+	}
+
+	/**
+	 * Checks to see if this module is loaded and ready.  It is only valid to call Get() if IsAvailable() returns true.
+	 *
+	 * @return True if the module is loaded and ready to use
+	 */
+	static inline bool IsAvailable()
+	{
+		return FModuleManager::Get().IsModuleLoaded( "MassSignals" );
+	}
+};
+
diff --git a/Engine/Source/Runtime/Mass/MassSignals/Public/MassSignalProcessorBase.h b/Engine/Source/Runtime/Mass/MassSignals/Public/MassSignalProcessorBase.h
new file mode 100644
--- /dev/null
+++ b/Engine/Source/Runtime/Mass/MassSignals/Public/MassSignalProcessorBase.h
@@ -0,0 +1,109 @@
+﻿// Copyright Epic Games, Inc. All Rights Reserved.
+
+#pragma once
+
+#include "MassProcessor.h"
+#include "MassEntityManager.h"
+#include "MassSignalTypes.h"
+#include "Containers/StaticArray.h"
+#include "Misc/TransactionallySafeRWLock.h"
+#include "MassEntityQuery.h"
+#include "MassSignalProcessorBase.generated.h"
+
+#define UE_API MASSSIGNALS_API
+
+class UMassSignalSubsystem;
+
+/** 
+ * Processor for executing signals on each targeted entities
+ * The derived classes only need to implement the method SignalEntities to actually received the raised signals for the entities they subscribed to 
+ */
+UCLASS(MinimalAPI, abstract)
+class UMassSignalProcessorBase : public UMassProcessor
+{
+	GENERATED_BODY()
+
+public:
+	UE_API UMassSignalProcessorBase(const FObjectInitializer& ObjectInitializer = FObjectInitializer::Get());
+
+protected:
+	UE_API virtual void BeginDestroy() override;
+
+	/** Configure the owned FMassEntityQuery instances to express processor queries requirements */
+	virtual void ConfigureQueries(const TSharedRef<FMassEntityManager>&) override {}
+
+	/**
+	 * Actual method that derived class needs to implement to act on a signal that is raised for that frame
+	 * @param EntitySubsystem is the system to execute the lambdas on each entity chunk
+	 * @param Context is the execution context to be passed when executing the lambdas
+	 * @param EntitySignals Look up to retrieve for each entities their raised signal via GetSignalsForEntity
+	 */
+	 virtual void SignalEntities(FMassEntityManager& EntityManager, FMassExecutionContext& Context, FMassSignalNameLookup& EntitySignals) PURE_VIRTUAL(UMassSignalProcessorBase::SignalEntities, );
+
+	/**
+	 * Callback that is being called when new signal is raised
+	 * @param SignalName is the name of the signal being raised
+	 * @param Entities are the targeted entities for this signal
+	 */
+	 UE_API virtual void OnSignalReceived(FName SignalName, TConstArrayView<FMassEntityHandle> Entities);
+
+	/**
+	 * Execution method for this processor
+	 * @param EntitySubsystem is the system to execute the lambdas on each entity chunk
+	 * @param Context is the execution context to be passed when executing the lambdas
+	 */
+	UE_API virtual void Execute(FMassEntityManager& EntityManager, FMassExecutionContext& Context) override;
+
+	/** 
+	 * We don't want signaling processors to be pruned by default, since that can mess up signal processing 
+	 * just after processor's creation (might miss some signals initially).
+	 */
+	virtual bool ShouldAllowQueryBasedPruning(const bool bRuntimeMode = true) const override { return false; }
+
+
+	/**
+	 * To receive notification about a particular signal, you need to subscribe to it.
+	 * @param SignalName is the name of the signal to receive notification about
+	 */
+	UE_API void SubscribeToSignal(UMassSignalSubsystem& SignalSubsystem, const FName SignalName);
+
+	FMassEntityQuery EntityQuery;
+
+private:
+
+	/** Stores a range of indices in the SignaledEntities TArray of Entities and the associated signal name */
+	struct FEntitySignalRange
+	{
+		FName SignalName;
+		int32 Begin = 0;
+		int32 End = 0;
+		bool bProcessed = false;
+	};
+
+	struct FFrameReceivedSignals
+	{
+		/** Received signals are double buffered as we can receive new one while processing them */
+		TArray<FEntitySignalRange> ReceivedSignalRanges;
+
+		/** the list of all signaled entities, can contain duplicates */
+		TArray<FMassEntityHandle> SignaledEntities;
+	};
+
+	static constexpr int BuffersCount = 2;
+
+	/** Double buffer frame received signal as we can receive new signals as we are processing them */
+	TStaticArray<FFrameReceivedSignals, BuffersCount> FrameReceivedSignals;
+	
+	/** Current frame buffer index of FrameReceivedSignals */
+	int32 CurrentFrameBufferIndex = 0;
+
+	/** Lookup used to store and retrieve signals per entity, only used during processing */
+	FMassSignalNameLookup SignalNameLookup;
+
+	/** List of all the registered signal names*/
+	TArray<FName> RegisteredSignals;
+
+	FTransactionallySafeRWLock ReceivedSignalLock;
+};
+
+#undef UE_API
diff --git a/Engine/Source/Runtime/Mass/MassSignals/Public/MassSignalSubsystem.h b/Engine/Source/Runtime/Mass/MassSignals/Public/MassSignalSubsystem.h
new file mode 100644
--- /dev/null
+++ b/Engine/Source/Runtime/Mass/MassSignals/Public/MassSignalSubsystem.h
@@ -0,0 +1,141 @@
+// Copyright Epic Games, Inc. All Rights Reserved.
+
+#pragma once
+
+#if UE_ENABLE_INCLUDE_ORDER_DEPRECATED_IN_5_6
+#include "MassEntityManager.h"
+#endif // UE_ENABLE_INCLUDE_ORDER_DEPRECATED_IN_5_6
+#include "MassEntityHandle.h"
+#include "MassSubsystemBase.h"
+#include "Misc/MTAccessDetector.h"
+#include "MassExternalSubsystemTraits.h"
+#include "MassSignalSubsystem.generated.h"
+
+struct FMassExecutionContext;
+
+namespace UE::MassSignal 
+{
+	DECLARE_MULTICAST_DELEGATE_TwoParams(FSignalDelegate, FName /*SignalName*/, TConstArrayView<FMassEntityHandle> /*Entities*/);
+} // UE::MassSignal
+
+/**
+* A subsystem for handling Signals in Mass
+*/
+UCLASS(MinimalAPI)
+class UMassSignalSubsystem : public UMassTickableSubsystemBase
+{
+	GENERATED_BODY()
+	
+public:
+
+	/** 
+	 * Retrieve the delegate dispatcher from the signal name
+	 * @param SignalName is the name of the signal to get the delegate dispatcher from
+	 */
+	UE::MassSignal::FSignalDelegate& GetSignalDelegateByName(FName SignalName)
+	{
+		return NamedSignals.FindOrAdd(SignalName);
+	}
+
+	/**
+	 * Inform a single entity of a signal being raised
+	 * @param SignalName is the name of the signal raised
+	 * @param Entity entity that should be informed that signal 'SignalName' was raised
+	 */
+	MASSSIGNALS_API void SignalEntity(FName SignalName, const FMassEntityHandle Entity);
+
+	/**
+	 * Inform multiple entities of a signal being raised
+	 * @param SignalName is the name of the signal raised
+	 * @param Entities list of entities that should be informed that signal 'SignalName' was raised
+	 */
+	MASSSIGNALS_API void SignalEntities(FName SignalName, TConstArrayView<FMassEntityHandle> Entities);
+
+	/**
+	 * Inform a single entity of a signal being raised in a certain amount of seconds
+	 * @param SignalName is the name of the signal raised
+	 * @param Entity entity that should be informed that signal 'SignalName' was raised
+	 * @param DelayInSeconds is the amount of time before signaling the entity
+	 */
+	MASSSIGNALS_API void DelaySignalEntity(FName SignalName, const FMassEntityHandle Entity, const float DelayInSeconds);
+
+ 	/**
+	 * Inform multiple entities of a signal being raised in a certain amount of seconds
+	 * @param SignalName is the name of the signal raised
+	 * @param Entities being informed of the raised signal
+	 * @param DelayInSeconds is the amount of time before signaling the entities
+	 */
+	MASSSIGNALS_API void DelaySignalEntities(FName SignalName, TConstArrayView<FMassEntityHandle> Entities, const float DelayInSeconds);
+
+	/**
+	 * Inform single entity of a signal being raised asynchronously using the Mass Command Buffer
+	 * @param Context is the Entity System execution context to push the async command
+	 * @param SignalName is the name of the signal raised
+	 * @param Entity entity that should be informed that signal 'SignalName' was raised
+	 */
+	MASSSIGNALS_API void SignalEntityDeferred(FMassExecutionContext& Context, FName SignalName, const FMassEntityHandle Entity);
+
+	/**
+	 * Inform multiple entities of a signal being raised asynchronously using the Mass Command Buffer
+	 * @param Context is the Entity System execution context to push the async command
+	 * @param SignalName is the name of the signal raised
+	 * @param Entities list of entities that should be informed that signal 'SignalName' was raised
+	 */
+	MASSSIGNALS_API void SignalEntitiesDeferred(FMassExecutionContext& Context, FName SignalName, TConstArrayView<FMassEntityHandle> Entities);
+
+ 	/**
+	 * Inform single entity of a signal being raised asynchronously using the Mass Command Buffer
+	 * @param Context is the Entity System execution context to push the async command
+	 * @param SignalName is the name of the signal raised
+	 * @param Entity entity that should be informed that signal 'SignalName' was raised
+	 * @param DelayInSeconds is the amount of time before signaling the entities
+	 */
+	MASSSIGNALS_API void DelaySignalEntityDeferred(FMassExecutionContext& Context, FName SignalName, const FMassEntityHandle Entity, const float DelayInSeconds);
+
+ 	/**
+	 * Inform multiple entities of a signal being raised asynchronously using the Mass Command Buffer
+	 * @param Context is the Entity System execution context to push the async command
+	 * @param SignalName is the name of the signal raised
+	 * @param Entities being informed of that signal was raised
+	 * @param DelayInSeconds is the amount of time before signaling the entities
+	 */
+	MASSSIGNALS_API void DelaySignalEntitiesDeferred(FMassExecutionContext& Context, FName SignalName, TConstArrayView<FMassEntityHandle> Entities, const float DelayInSeconds);
+
+protected:
+	// USubsystem implementation Begin
+	MASSSIGNALS_API virtual void Initialize(FSubsystemCollectionBase& Collection) override;
+	MASSSIGNALS_API virtual void Deinitialize() override;
+	// USubsystem implementation End	
+
+	MASSSIGNALS_API virtual void Tick(float DeltaTime) override;
+	MASSSIGNALS_API virtual TStatId GetStatId() const override;
+
+	/** Multithreading access detector to validate accesses to the list of delayed signals */
+	UE_MT_DECLARE_RW_ACCESS_DETECTOR(DelayedSignalsAccessDetector);
+
+	TMap<FName, UE::MassSignal::FSignalDelegate> NamedSignals;
+
+	struct FDelayedSignal
+	{
+		FName SignalName;
+		TArray<FMassEntityHandle> Entities;
+		double TargetTimestamp;
+	};
+
+	TArray<FDelayedSignal> DelayedSignals;
+
+	UPROPERTY(transient)
+	TObjectPtr<UWorld> CachedWorld;
+};
+
+template<>
+struct TMassExternalSubsystemTraits<UMassSignalSubsystem> final
+{
+	enum
+	{
+		GameThreadOnly = false,
+		// @todo this subsystem not being thread-safe when writing is an obstacle in
+		// parallelizing multiple processors
+		ThreadSafeWrite = false,
+	};
+};
diff --git a/Engine/Source/Runtime/Mass/MassSignals/Public/MassSignalTypes.h b/Engine/Source/Runtime/Mass/MassSignals/Public/MassSignalTypes.h
new file mode 100644
--- /dev/null
+++ b/Engine/Source/Runtime/Mass/MassSignals/Public/MassSignalTypes.h
@@ -0,0 +1,51 @@
+// Copyright Epic Games, Inc. All Rights Reserved.
+
+#pragma once
+
+#if UE_ENABLE_INCLUDE_ORDER_DEPRECATED_IN_5_6
+#include "CoreMinimal.h"
+#endif // UE_ENABLE_INCLUDE_ORDER_DEPRECATED_IN_5_6
+#include "Logging/LogMacros.h"
+#include "MassEntityManager.h"
+
+DECLARE_LOG_CATEGORY_EXTERN(LogMassSignals, Log, All)
+
+/**
+ * MassSignalNameLookup stores list of Signal names for each entity. The names are stored per entity as a bitmask,
+ * you can allocate new name using GetOrAddSignalName(). This limits the names to 64.
+ */
+struct FMassSignalNameLookup
+{
+	/** Max number of names each entity can contain */
+	static constexpr int32 MaxSignalNames = 64;
+
+	/** 
+	 * Retrieve if it is already registered or adds new signal to the lookup and return the bitflag for that Signal
+	 * @SignalName is the name of the signal to retrieve or add to the lookup.
+	 * @return bitflag describing the name, or 0 if max names is reached. 
+	 */
+	MASSSIGNALS_API uint64 GetOrAddSignalName(const FName SignalName);
+
+	/**
+	 * Adds specified Signal name bitflag to an entity 
+	 * @param Entity is the entity where the signal has been raised
+	 * @param SignalFlag is the actual bitflag describing the signal
+	 */
+	MASSSIGNALS_API void AddSignalToEntity(const FMassEntityHandle Entity, const uint64 SignalFlag);
+
+	/** 
+	 * Retrieve for a specific entity the raised signal this frame
+	 * @return Array of signal names raised for this entity 
+	 */
+	MASSSIGNALS_API void GetSignalsForEntity(const FMassEntityHandle Entity, TArray<FName>& OutSignals) const;
+
+	/** Empties the name lookup and entity signals */
+	MASSSIGNALS_API void Reset();
+
+protected:
+	/** Array of Signal names */
+	TArray<FName> SignalNames;
+
+	/** Map from entity id to name bitmask */
+	TMap<FMassEntityHandle, uint64> EntitySignals;
+};
