MassSignalSubsystem's Tick is now able to work with AutoRTFM enabled (by Ticking after the transaction).

#rnx

diff --git a/Engine/Source/Runtime/Mass/MassSignals/Private/MassSignalSubsystem.cpp b/Engine/Source/Runtime/Mass/MassSignals/Private/MassSignalSubsystem.cpp
--- a/Engine/Source/Runtime/Mass/MassSignals/Private/MassSignalSubsystem.cpp
+++ b/Engine/Source/Runtime/Mass/MassSignals/Private/MassSignalSubsystem.cpp
@@ -30,6 +30,20 @@ void UMassSignalSubsystem::Deinitialize()
 
 void UMassSignalSubsystem::Tick(float DeltaTime)
 {
+	// If we're in an AutoRTFM Transaction, exit the TickableObject Tick (which is top-level) and run again outside the transaction
+	if (AutoRTFM::IsTransactional())
+	{
+		AutoRTFM::OnCommit([WeakThis = TWeakObjectPtr(this), DeltaTime]()
+		{
+			if (ensure(WeakThis.IsValid()))
+			{
+				WeakThis->Tick(DeltaTime);
+			}
+		});
+
+		return;
+	}
+
 	CA_ASSUME(CachedWorld);
 	const double CurrentTime = CachedWorld->GetTimeSeconds();
 
