Fix build break

#rnx

diff --git a/Engine/Source/Runtime/Mass/MassEngine/Private/AI/MassEngineNavigationProcessors.cpp b/Engine/Source/Runtime/Mass/MassEngine/Private/AI/MassEngineNavigationProcessors.cpp
new file mode 100644
--- /dev/null
+++ b/Engine/Source/Runtime/Mass/MassEngine/Private/AI/MassEngineNavigationProcessors.cpp
@@ -0,0 +1,323 @@
+// Copyright Epic Games, Inc. All Rights Reserved.
+
+#include "AI/MassEngineNavigationProcessors.h"
+
+#include "Engine/StaticMesh.h"
+#include "MassEntityFragments.h"
+#include "MassEngineTypes.h"
+#include "MassExecutionContext.h"
+#include "Mesh/MassEngineMeshFragments.h"
+#include "Physics/MassEnginePhysicsFragments.h"
+#include "MassSignalSubsystem.h"
+
+#include "AI/NavigationSystemBase.h"
+#include "AI/Navigation/NavCollisionBase.h"
+#include "AI/Navigation/NavigationElement.h"
+#include "AI/Navigation/NavigationRelevantData.h"
+
+// Helper
+
+#include UE_INLINE_GENERATED_CPP_BY_NAME(MassEngineNavigationProcessors)
+
+namespace UE::NavigationRelevant
+{
+
+bool IsNavigationRelevant(const FMassColliderFragment& ColliderFragment,
+	const FMassStaticMeshFragment& MeshFragment,
+	const FTransformFragment& TransformFragment)
+{
+	if (TransformFragment.GetTransform().GetScale3D().IsNearlyZero())
+	{
+		return false;
+	}
+
+	// @see UPrimitiveComponent::IsNavigationRelevant()
+	{
+		if (!CollisionEnabledHasQuery(ColliderFragment.CollisionType))
+		{
+			return false;
+		}
+
+		const FCollisionResponseContainer& ResponseToChannels = ColliderFragment.CollisionResponse;
+		if (!(ResponseToChannels.GetResponse(ECC_Pawn) == ECR_Block
+			|| ResponseToChannels.GetResponse(ECC_Vehicle) == ECR_Block))
+		{
+			return false;
+		}
+	}
+
+	// @see UStaticMeshComponent::IsNavigationRelevant()
+	const UStaticMesh* Mesh = MeshFragment.Mesh.Get();
+	return Mesh != nullptr && !Mesh->IsCompiling() && Mesh->IsNavigationRelevant();
+}
+
+void GetNavigationData(const FNavigationElement& Element, const FMassNavigationRelevantParameters& Parameters, const TWeakObjectPtr<UNavCollisionBase> WeakNavCollision, FNavigationRelevantData& OutData)
+{
+	// @see UPrimitiveComponent::GetNavigationData
+	if (Parameters.bFillCollisionUnderneathForNavData)
+	{
+		FCompositeNavModifier CompositeNavModifier;
+		CompositeNavModifier.SetFillCollisionUnderneathForNavmesh(true);
+		OutData.Modifiers.Add(CompositeNavModifier);
+	}
+
+	// @see UStaticMeshComponent::GetNavigationData
+	// but here we don't expose options to override the export (i.e., bOverrideNavigationExport & bForceNavigationObstacle)
+	if (UNavCollisionBase* NavCollision = WeakNavCollision.Get())
+	{
+		if (NavCollision->IsDynamicObstacle())
+		{
+			NavCollision->GetNavigationModifier(OutData.Modifiers, Element.GetTransform());
+		}
+
+		// @todo: allow nav collision to use specific area class
+		//if (NavArea)
+		//{
+		//	NavCollision->GetNavigationModifier(Data.Modifiers, Transform);
+		//	TArray<FAreaNavModifier> Areas = Data.Modifiers.GetMutableAreas();
+		//	for (FAreaNavModifier& Area : Areas)
+		//	{
+		//		Area.SetAreaClass(NavArea);
+		//	}
+		//}
+	}
+}
+
+bool DoCustomNavigableGeometryExport(const FNavigationElement& Element, const TWeakObjectPtr<UNavCollisionBase> WeakNavCollision, FNavigableGeometryExport& OutGeomExport)
+{
+	if (const UNavCollisionBase* NavCollision = WeakNavCollision.Get())
+	{
+		// Dynamic obstacles are handled from GetNavigationData using NavModifiers
+		if (NavCollision->IsDynamicObstacle())
+		{
+			// returning 'false' to skip call to the default export (i.e. ExportRigidBodySetup) since we only use the
+			// geometry to mark an obstacle.
+			return false;
+		}
+
+		if (const bool bHasExportedGeometry = NavCollision->ExportGeometry(Element.GetTransform(), OutGeomExport))
+		{
+			// returning 'false' to skip call to the default export (i.e. ExportRigidBodySetup) since the export
+			// was handled by the NavCollision.
+			return false;
+		}
+	}
+
+	return true;
+}
+
+void ConfigureElement(
+	const UObject* Owner,
+	FNavigationElement& OutElement,
+	const FMassStaticMeshFragment& MeshFragment,
+	const FTransformFragment& TransformFragment,
+	const FMassNavigationRelevantParameters& NavigationRelevantParameters)
+{
+	const UStaticMesh* StaticMesh = MeshFragment.Mesh.Get();
+	const FTransform& Transform = TransformFragment.GetTransform();
+	if (StaticMesh)
+	{
+		OutElement.SetBodySetup(StaticMesh->GetBodySetup());
+		OutElement.SetBounds(StaticMesh->GetNavigationBounds(Transform));
+	}
+	OutElement.SetTransform(Transform);
+
+	OutElement.NavigationDataExportDelegate.BindWeakLambda(Owner,
+		[NavigationRelevantParameters, WeakNavCollision = TWeakObjectPtr<UNavCollisionBase>(StaticMesh ? StaticMesh->GetNavCollision() : nullptr)]
+		(const FNavigationElement& Element, FNavigationRelevantData& OutNavigationRelevantData)
+		{
+			GetNavigationData(
+				Element,
+				NavigationRelevantParameters,
+				WeakNavCollision,
+				OutNavigationRelevantData);
+		});
+	OutElement.CustomGeometryExportDelegate.BindWeakLambda(Owner,
+		[WeakNavCollision = TWeakObjectPtr<UNavCollisionBase>(StaticMesh ? StaticMesh->GetNavCollision() : nullptr)]
+		(const FNavigationElement& Element, FNavigableGeometryExport& OutGeometry, bool& bOutExportDefaultGeometry)
+		{
+			bOutExportDefaultGeometry = DoCustomNavigableGeometryExport(
+				Element,
+				WeakNavCollision,
+				OutGeometry);
+		});
+}
+} // UE::NavigationRelevant
+
+//----------------------------------------------------------------------//
+// UMassNavigationElementProcessor
+//----------------------------------------------------------------------//
+UMassNavigationElementProcessor::UMassNavigationElementProcessor(const FObjectInitializer& ObjectInitializer)
+	: Super(ObjectInitializer)
+	, NewRelevantEntityQuery(*this)
+{
+	bRequiresGameThreadExecution = true;
+	ExecutionFlags = static_cast<int32>(EProcessorExecutionFlags::AllWorldModes);
+}
+
+void UMassNavigationElementProcessor::ConfigureQueries(const TSharedRef<FMassEntityManager>& EntityManager)
+{
+	NewRelevantEntityQuery.AddConstSharedRequirement<FMassColliderFragment>();
+	NewRelevantEntityQuery.AddConstSharedRequirement<FMassStaticMeshFragment>();
+	NewRelevantEntityQuery.AddRequirement<FTransformFragment>(EMassFragmentAccess::ReadOnly);
+	NewRelevantEntityQuery.AddRequirement<FMassNavigationRelevantFragment>(EMassFragmentAccess::ReadOnly, EMassFragmentPresence::None);
+}
+
+void UMassNavigationElementProcessor::Execute(FMassEntityManager& EntityManager, FMassExecutionContext& ExecutionContext)
+{
+	using namespace UE::NavigationRelevant;
+
+	const UObject* EntityManagerOwner = ExecutionContext.GetEntityManagerChecked().GetOwner();
+	check(EntityManagerOwner);
+
+	// @todo: hardcoded for now but should use some default settings and eventually add dedicated component to override settings
+	FMassNavigationRelevantParameters Parameters;
+	Parameters.bFillCollisionUnderneathForNavData = false;
+	const FConstSharedStruct ParametersSharedFragment = EntityManager.GetOrCreateConstSharedFragment(Parameters);
+
+	NewRelevantEntityQuery.ForEachEntityChunk(ExecutionContext,
+		[EntityManagerOwner, ParametersSharedFragment](FMassExecutionContext& Context)
+		{
+			const TConstArrayView<FTransformFragment> TransformFragments = Context.GetFragmentView<FTransformFragment>();
+			const FMassStaticMeshFragment& MeshFragment = Context.GetConstSharedFragment<FMassStaticMeshFragment>();
+			const FMassColliderFragment& ColliderFragment = Context.GetConstSharedFragment<FMassColliderFragment>();
+
+			for (FMassExecutionContext::FEntityIterator EntityIt = Context.CreateEntityIterator(); EntityIt; ++EntityIt)
+			{
+				const FMassEntityHandle EntityHandle = Context.GetEntity(EntityIt);
+				FNavigationElementHandle NavigationElementHandle = FNavigationElementHandle::Invalid;
+
+				if (IsNavigationRelevant(ColliderFragment, MeshFragment, TransformFragments[EntityIt]))
+				{
+					FNavigationElement Element(*EntityManagerOwner, EntityHandle.AsNumber());
+					ConfigureElement(EntityManagerOwner, Element, MeshFragment, TransformFragments[EntityIt], ParametersSharedFragment.Get<const FMassNavigationRelevantParameters>());
+					NavigationElementHandle = FNavigationSystem::AddNavigationElement(EntityManagerOwner->GetWorld(), MoveTemp(Element));
+				}
+
+				Context.Defer().PushCommand<FMassCommandAddFragmentInstances>(EntityHandle, FMassNavigationRelevantFragment(NavigationElementHandle));
+				Context.Defer().PushCommand<FMassDeferredAddCommand>([EntityHandle, ParametersSharedFragment](FMassEntityManager& InEntityManager)
+					{
+						if (InEntityManager.IsEntityValid(EntityHandle))
+						{
+							InEntityManager.AddConstSharedFragmentToEntity(EntityHandle, ParametersSharedFragment);
+						}
+						else
+						{
+							UE_LOG(LogMass, Log, TEXT("FMassNavigationRelevantParameters shared fragment not added: Entity is no longer valid"));
+						}
+					});
+			}
+		});
+}
+
+//----------------------------------------------------------------------//
+// UMassDirtyNavigationRelevantUpdateProcessor
+//----------------------------------------------------------------------//
+UMassDirtyNavigationRelevantUpdateProcessor::UMassDirtyNavigationRelevantUpdateProcessor(const FObjectInitializer& ObjectInitializer)
+	: Super(ObjectInitializer)
+{
+	bRequiresGameThreadExecution = true;
+	ExecutionFlags = static_cast<int32>(EProcessorExecutionFlags::AllWorldModes);
+}
+
+void UMassDirtyNavigationRelevantUpdateProcessor::InitializeInternal(UObject& Owner, const TSharedRef<FMassEntityManager>& EntityManager)
+{
+	Super::InitializeInternal(Owner, EntityManager);
+
+	if (const UWorld* World = GetWorld())
+	{
+		if (UMassSignalSubsystem* SignalSubsystem = World->GetSubsystem<UMassSignalSubsystem>())
+		{
+			SubscribeToSignal(*SignalSubsystem, UE::Mass::Signals::TransformChanged);
+			SubscribeToSignal(*SignalSubsystem, UE::Mass::Signals::MeshChanged);
+		}
+	}
+}
+
+void UMassDirtyNavigationRelevantUpdateProcessor::ConfigureQueries(const TSharedRef<FMassEntityManager>& EntityManager)
+{
+	EntityQuery.AddConstSharedRequirement<FMassColliderFragment>();
+	EntityQuery.AddConstSharedRequirement<FMassStaticMeshFragment>();
+	EntityQuery.AddConstSharedRequirement<FMassNavigationRelevantParameters>();
+	EntityQuery.AddRequirement<FTransformFragment>(EMassFragmentAccess::ReadOnly);
+	EntityQuery.AddRequirement<FMassNavigationRelevantFragment>(EMassFragmentAccess::ReadWrite);
+}
+
+void UMassDirtyNavigationRelevantUpdateProcessor::SignalEntities(
+	FMassEntityManager& EntityManager,
+	FMassExecutionContext& ExecutionContext,
+	FMassSignalNameLookup& EntitySignals)
+{
+	using namespace UE::NavigationRelevant;
+
+	const UObject* EntityManagerOwner = ExecutionContext.GetEntityManagerChecked().GetOwner();
+	check(EntityManagerOwner);
+
+	EntityQuery.ForEachEntityChunk(ExecutionContext,
+		[EntityManagerOwner, World = GetWorld()](FMassExecutionContext& Context)
+		{
+			const TConstArrayView<FTransformFragment> TransformFragments = Context.GetFragmentView<FTransformFragment>();
+			const TArrayView<FMassNavigationRelevantFragment> NavigationRelevantFragments = Context.GetMutableFragmentView<FMassNavigationRelevantFragment>();
+
+			const FMassColliderFragment& ColliderFragment = Context.GetConstSharedFragment<FMassColliderFragment>();
+			const FMassStaticMeshFragment& MeshFragment = Context.GetConstSharedFragment<FMassStaticMeshFragment>();
+			const FMassNavigationRelevantParameters& NavigationRelevantParameters = Context.GetConstSharedFragment<FMassNavigationRelevantParameters>();
+
+			for (FMassExecutionContext::FEntityIterator EntityIt = Context.CreateEntityIterator(); EntityIt; ++EntityIt)
+			{
+				const bool bWasRelevant = NavigationRelevantFragments[EntityIt].Handle.IsValid();
+
+				if (IsNavigationRelevant(ColliderFragment, MeshFragment, TransformFragments[EntityIt]))
+				{
+					FNavigationElement Element(*EntityManagerOwner, Context.GetEntity(EntityIt).AsNumber());
+					ConfigureElement(EntityManagerOwner, Element, MeshFragment, TransformFragments[EntityIt], NavigationRelevantParameters);
+
+					if (bWasRelevant)
+					{
+						FNavigationSystem::UpdateNavigationElement(World, NavigationRelevantFragments[EntityIt].Handle, MoveTemp(Element));
+					}
+					else
+					{
+						NavigationRelevantFragments[EntityIt].Handle = FNavigationSystem::AddNavigationElement(World, MoveTemp(Element));
+					}
+				}
+				else if (bWasRelevant)
+				{
+					FNavigationSystem::RemoveNavigationElement(World, NavigationRelevantFragments[EntityIt].Handle);
+				}
+			}
+		});
+}
+
+//----------------------------------------------------------------------//
+// UMassNavigationRelevantMassDeinitializer
+//----------------------------------------------------------------------//
+UMassNavigationRelevantMassDeinitializer::UMassNavigationRelevantMassDeinitializer()
+	: EntityQuery(*this)
+{
+	ExecutionFlags = static_cast<int32>(EProcessorExecutionFlags::AllWorldModes);
+	ObservedType = FMassNavigationRelevantFragment::StaticStruct();
+	ObservedOperations = EMassObservedOperationFlags::Remove;
+	bRequiresGameThreadExecution = true;
+}
+
+void UMassNavigationRelevantMassDeinitializer::ConfigureQueries(const TSharedRef<FMassEntityManager>& EntityManager)
+{
+	EntityQuery.AddRequirement<FMassNavigationRelevantFragment>(EMassFragmentAccess::ReadOnly);
+}
+
+void UMassNavigationRelevantMassDeinitializer::Execute(FMassEntityManager& EntityManager, FMassExecutionContext& ExecutionContext)
+{
+	EntityQuery.ForEachEntityChunk(ExecutionContext,
+		[World = GetWorld()](FMassExecutionContext& Context)
+		{
+			const TConstArrayView<FMassNavigationRelevantFragment> NavigationRelevantFragments = Context.GetFragmentView<FMassNavigationRelevantFragment>();
+			for (FMassExecutionContext::FEntityIterator EntityIt = Context.CreateEntityIterator(); EntityIt; ++EntityIt)
+			{
+				if (NavigationRelevantFragments[EntityIt].Handle.IsValid())
+				{
+					FNavigationSystem::RemoveNavigationElement(World, NavigationRelevantFragments[EntityIt].Handle);
+				}
+			}
+		});
+}
diff --git a/Engine/Source/Runtime/Mass/MassEngine/Private/Mesh/MassEngineMeshFragments.cpp b/Engine/Source/Runtime/Mass/MassEngine/Private/Mesh/MassEngineMeshFragments.cpp
new file mode 100644
--- /dev/null
+++ b/Engine/Source/Runtime/Mass/MassEngine/Private/Mesh/MassEngineMeshFragments.cpp
@@ -0,0 +1,8 @@
+// Copyright Epic Games, Inc. All Rights Reserved.
+
+#include "Mesh/MassEngineMeshFragments.h"
+
+FMassStaticMeshFragment::FMassStaticMeshFragment(TNotNull<const UStaticMesh*> InMesh)
+	: Mesh(InMesh)
+{
+}
diff --git a/Engine/Source/Runtime/Mass/MassEngine/Public/AI/MassEngineNavigationProcessors.h b/Engine/Source/Runtime/Mass/MassEngine/Public/AI/MassEngineNavigationProcessors.h
new file mode 100644
--- /dev/null
+++ b/Engine/Source/Runtime/Mass/MassEngine/Public/AI/MassEngineNavigationProcessors.h
@@ -0,0 +1,71 @@
+// Copyright Epic Games, Inc. All Rights Reserved.
+
+#pragma once
+
+#include "MassObserverProcessor.h"
+#include "MassSignalProcessorBase.h"
+#include "AI/MassEngineNavigationFragments.h"
+
+#include "MassEngineNavigationProcessors.generated.h"
+
+struct FMassEntityManager;
+struct FMassExecutionContext;
+
+/**
+ * Processor responsible for registering new NavigationElement representing Mass entities that are relevant to the AI navigation system
+ * and adding the FMassNavigationRelevantFragment fragment to the entities.
+ */
+UCLASS()
+class UMassNavigationElementProcessor : public UMassProcessor
+{
+	GENERATED_BODY()
+
+public:
+	UMassNavigationElementProcessor(const FObjectInitializer& ObjectInitializer = FObjectInitializer::Get());
+
+protected:
+	virtual void ConfigureQueries(const TSharedRef<FMassEntityManager>& EntityManager) override;
+	virtual void Execute(FMassEntityManager& EntityManager, FMassExecutionContext& Context) override;
+
+	FMassEntityQuery NewRelevantEntityQuery;
+};
+
+/**
+ * Processor responsible for updating NavigationElement representing Mass entities that are relevant to the AI navigation system
+ * when one of its core property gets modified (e.g., transform, mesh, navigation bounds, etc.)
+ * This processor subscribes to the following signals to perform an update:
+ * - UE::Mass::Signals::TransformChanged
+ * - UE::Mass::Signals::MeshChanged
+ */
+UCLASS()
+class UMassDirtyNavigationRelevantUpdateProcessor : public UMassSignalProcessorBase
+{
+	GENERATED_BODY()
+
+public:
+	UMassDirtyNavigationRelevantUpdateProcessor(const FObjectInitializer& ObjectInitializer = FObjectInitializer::Get());
+
+protected:
+	virtual void InitializeInternal(UObject& Owner, const TSharedRef<FMassEntityManager>& EntityManager) override;
+	virtual void ConfigureQueries(const TSharedRef<FMassEntityManager>& EntityManager) override;
+	virtual void SignalEntities(FMassEntityManager& EntityManager, FMassExecutionContext& Context, FMassSignalNameLookup& EntitySignals) override;
+};
+
+/**
+ * Processor responsible for unregistering the NavigationElements when Mass entities are destroyed or no longer relevant
+ * to the AI navigation system.
+ */
+UCLASS()
+class UMassNavigationRelevantMassDeinitializer : public UMassObserverProcessor
+{
+	GENERATED_BODY()
+
+public:
+	UMassNavigationRelevantMassDeinitializer();
+
+protected:
+	virtual void ConfigureQueries(const TSharedRef<FMassEntityManager>& EntityManager) override;
+	virtual void Execute(FMassEntityManager& EntityManager, FMassExecutionContext& Context) override;
+
+	FMassEntityQuery EntityQuery;
+};
\ No newline at end of file
diff --git a/Engine/Source/Runtime/Mass/MassEngine/Public/Mesh/MassEngineMeshFragments.h b/Engine/Source/Runtime/Mass/MassEngine/Public/Mesh/MassEngineMeshFragments.h
new file mode 100644
--- /dev/null
+++ b/Engine/Source/Runtime/Mass/MassEngine/Public/Mesh/MassEngineMeshFragments.h
@@ -0,0 +1,30 @@
+﻿// Copyright Epic Games, Inc. All Rights Reserved.
+
+#pragma once
+
+#include "MassEntityElementTypes.h"
+#include "Engine/StaticMesh.h"
+
+#include "MassEngineMeshFragments.generated.h"
+
+#define UE_API MASSENGINE_API
+
+/**
+ * Fragment holding a pointer to the static mesh associated to the entity
+ */
+USTRUCT()
+struct FMassStaticMeshFragment : public FMassConstSharedFragment
+{
+	GENERATED_BODY()
+
+	FMassStaticMeshFragment() = default;
+	UE_API explicit FMassStaticMeshFragment(TNotNull<const UStaticMesh*> InMesh);
+
+	UPROPERTY(Transient, VisibleAnywhere, Category=Debug)
+	TWeakObjectPtr<const UStaticMesh> Mesh;
+
+	UPROPERTY(Transient, VisibleAnywhere, Category = Debug)
+	TArray<TWeakObjectPtr<UMaterialInterface>> OverrideMaterials;
+};
+
+#undef UE_API
\ No newline at end of file
