[EntityFramework] added more logging to mass_component under new LogCategory "LogSceneGraphMass"
Also added new LogCategory "LogMassEngine"

#rnx
#Synced-CL 51261467

diff --git a/Engine/Source/Runtime/Mass/MassEngine/Private/MassEngineTypes.cpp b/Engine/Source/Runtime/Mass/MassEngine/Private/MassEngineTypes.cpp
new file mode 100644
--- /dev/null
+++ b/Engine/Source/Runtime/Mass/MassEngine/Private/MassEngineTypes.cpp
@@ -0,0 +1,5 @@
+// Copyright Epic Games, Inc. All Rights Reserved.
+
+#include "MassEngineTypes.h"
+
+DEFINE_LOG_CATEGORY(LogMassEngine);
\ No newline at end of file
diff --git a/Engine/Source/Runtime/Mass/MassEngine/Public/MassEngineTypes.h b/Engine/Source/Runtime/Mass/MassEngine/Public/MassEngineTypes.h
--- a/Engine/Source/Runtime/Mass/MassEngine/Public/MassEngineTypes.h
+++ b/Engine/Source/Runtime/Mass/MassEngine/Public/MassEngineTypes.h
@@ -2,8 +2,11 @@
 
 #pragma once
 
+#include "Logging/LogCategory.h"
 #include "UObject/NameTypes.h"
 
+MASSENGINE_API DECLARE_LOG_CATEGORY_EXTERN(LogMassEngine, Warning, All);
+
 namespace UE::Mass::Signals
 {
 	const FName TransformChanged = FName(TEXT("TransformChanged"));
