First pass of Mass API for rendering. It is not perfect but good enough to have more people start contributing to it.
This include rendering of of staticmesh and also auto instanced static mesh.
An other changelist will follow that uses this API for scenegraph.

#Synced-CL 51263739

diff --git a/Engine/Source/Runtime/Engine/Classes/Components/ActorComponent.h b/Engine/Source/Runtime/Engine/Classes/Components/ActorComponent.h
--- a/Engine/Source/Runtime/Engine/Classes/Components/ActorComponent.h
+++ b/Engine/Source/Runtime/Engine/Classes/Components/ActorComponent.h
@@ -48,8 +48,18 @@ class FRegisterComponentContext
 public:
 	using FPrimitiveBatch = TArray<UPrimitiveComponent*, FConcurrentLinearArrayAllocator>;
 
-	FRegisterComponentContext(UWorld* InWorld);
-	~FRegisterComponentContext();
+	ENGINE_API FRegisterComponentContext(UWorld* InWorld);
+	ENGINE_API ~FRegisterComponentContext();
+
+	bool IsInitialized() const
+	{
+		return World != nullptr;
+	}
+	void Initialize(UWorld* InWorld)
+	{
+		checkf(!IsInitialized(), TEXT("Context is already initialized"));
+		World = InWorld;
+	}
 
 	FRegisterComponentContext& SetIncrementalUpdateTimeout(const UE::FTimeout& InTimeout);
 	FRegisterComponentContext& SetIncrementalUpdateAddPrimitiveBatchSize(int32 Value);
@@ -68,12 +78,12 @@ public:
 	UE_DEPRECATED(5.5, "This function is not used anymore")
 	int32 Count() const { return AddPrimitiveBatches.Num(); }
 
-	void Process();
+	ENGINE_API void Process();
 
 private:
 	bool OnIncrementalRegisterComponentsDone();
 
-	UWorld* World;
+	UWorld* World = nullptr;
 	UE::FTimeout Timeout;
 	int32	IncrementalUpdateAddPrimitiveBatchSize;
 	int32	IncrementalUpdateNumComponentsToUpdate;
diff --git a/Engine/Source/Runtime/Engine/Public/PrimitiveSceneProxy.h b/Engine/Source/Runtime/Engine/Public/PrimitiveSceneProxy.h
--- a/Engine/Source/Runtime/Engine/Public/PrimitiveSceneProxy.h
+++ b/Engine/Source/Runtime/Engine/Public/PrimitiveSceneProxy.h
@@ -346,12 +346,12 @@ public:
 	/**
 	 * Force the proxy to render as though it is selected. Shows outlines and overlays.
 	 */
-	void SetSelectionOverride_GameThread(bool bForceSelection);
+	ENGINE_API void SetSelectionOverride_GameThread(bool bForceSelection);
 	
 	/**
 	 * Enqueue updated selection outline color for the render thread to use.
 	 */
-	void SetSelectionOutlineColorIndex_GameThread(uint8 ColorIndex);
+	ENGINE_API void SetSelectionOutlineColorIndex_GameThread(uint8 ColorIndex);
 
 	/**
 	 * Enqueue updated overlay color for the render thread to use.
diff --git a/Engine/Source/Runtime/Mass/MassEngine/MassEngine.Build.cs b/Engine/Source/Runtime/Mass/MassEngine/MassEngine.Build.cs
--- a/Engine/Source/Runtime/Mass/MassEngine/MassEngine.Build.cs
+++ b/Engine/Source/Runtime/Mass/MassEngine/MassEngine.Build.cs
@@ -19,8 +19,19 @@ namespace UnrealBuildTool.Rules
 					"Engine",
 					"MassEntity",
 					"MassSignals",
+					"RenderCore",
+					"RHI",
 				}
 			);
+
+			if (Target.bBuildEditor)
+			{
+				PrivateDependencyModuleNames.AddRange(
+					new [] {
+						"UnrealEd",
+					}
+				);
+			}
 		}
 	}
 }
diff --git a/Engine/Source/Runtime/Mass/MassEngine/Private/Mesh/MassDestroyRenderStateContext.cpp b/Engine/Source/Runtime/Mass/MassEngine/Private/Mesh/MassDestroyRenderStateContext.cpp
new file mode 100644
--- /dev/null
+++ b/Engine/Source/Runtime/Mass/MassEngine/Private/Mesh/MassDestroyRenderStateContext.cpp
@@ -0,0 +1,80 @@
+// Copyright Epic Games, Inc. All Rights Reserved.
+
+#include "MassDestroyRenderStateContext.h"
+#include "SceneProxies/DeferredDecalProxy.h"
+#include "PrimitiveSceneProxy.h"
+#include "SceneInterface.h"
+
+FMassDestroyRenderStateContext::FMassDestroyRenderStateContext(FSceneInterface* InScene)
+	: Scene(InScene)
+{
+}
+
+FMassDestroyRenderStateContext::~FMassDestroyRenderStateContext()
+{
+	Process();
+}
+
+void FMassDestroyRenderStateContext::Process()
+{
+	if (!PrimitiveSceneProxies.IsEmpty())
+	{
+		checkf(Scene, TEXT("Nothing should have been queued up if the scene was not set"));
+		Scene->BatchRemovePrimitives(MoveTemp(PrimitiveSceneProxies));
+		PrimitiveSceneProxies.Reset();
+	}
+	if (!DeferredDecalProxies.IsEmpty())
+	{
+		checkf(Scene, TEXT("Nothing should have been queued up if the scene was not set"));
+		TArray<FDeferredDecalUpdateParams> BatchParams;
+		BatchParams.Reserve(DeferredDecalProxies.Num());
+		for (FDeferredDecalProxy* Proxy : DeferredDecalProxies)
+		{
+			BatchParams.Add(
+				{
+					.DecalProxy = Proxy,
+					.OperationType = FDeferredDecalUpdateParams::EOperationType::RemoveFromSceneAndDelete
+				});
+		}
+		Scene->BatchUpdateDecals(MoveTemp(BatchParams));
+		DeferredDecalProxies.Reset();
+	}
+}
+
+void FMassDestroyRenderStateContext::DestroyProxy(FMassDestroyRenderStateContext* InContext, FPrimitiveSceneProxy* InPrimitiveSceneProxy, FSceneInterface* InScene)
+{
+	check(InScene);
+	check(InPrimitiveSceneProxy);
+
+	if (InContext)
+	{
+		checkf(InScene == InContext->Scene, TEXT("Expecting the scene to match the context one"));
+		InContext->PrimitiveSceneProxies.Add(InPrimitiveSceneProxy);
+	}
+	else
+	{
+		InScene->BatchRemovePrimitives({ InPrimitiveSceneProxy });
+	}
+}
+
+void FMassDestroyRenderStateContext::DestroyProxy(FMassDestroyRenderStateContext* InContext, FDeferredDecalProxy* InDeferredDecalProxy, FSceneInterface* InScene)
+{
+	check(InScene);
+	check(InDeferredDecalProxy);
+
+	if (InContext)
+	{
+		checkf(InScene == InContext->Scene, TEXT("Expecting the scene to match the context one"));
+		InContext->DeferredDecalProxies.Add(InDeferredDecalProxy);
+	}
+	else
+	{
+		TArray<FDeferredDecalUpdateParams> BatchParams;
+		BatchParams.Add(
+		{
+			.DecalProxy = InDeferredDecalProxy,
+			.OperationType = FDeferredDecalUpdateParams::EOperationType::RemoveFromSceneAndDelete
+		});
+		InScene->BatchUpdateDecals(MoveTemp(BatchParams));
+	}
+}
\ No newline at end of file
diff --git a/Engine/Source/Runtime/Mass/MassEngine/Private/Mesh/MassDestroyRenderStateContext.h b/Engine/Source/Runtime/Mass/MassEngine/Private/Mesh/MassDestroyRenderStateContext.h
new file mode 100644
--- /dev/null
+++ b/Engine/Source/Runtime/Mass/MassEngine/Private/Mesh/MassDestroyRenderStateContext.h
@@ -0,0 +1,40 @@
+// Copyright Epic Games, Inc. All Rights Reserved.
+
+#pragma once
+
+#include "Containers/Array.h"
+
+class FSceneInterface;
+class FDeferredDecalProxy;
+class FPrimitiveSceneProxy;
+
+/*
+ * Context to batch the destroy render states.
+ */
+struct FMassDestroyRenderStateContext
+{
+public:
+	FMassDestroyRenderStateContext(FSceneInterface* InScene);
+	~FMassDestroyRenderStateContext();
+
+	bool IsInitialized() const
+	{
+		return Scene != nullptr;
+	}
+
+	void Initialize(FSceneInterface* InScene)
+	{
+		checkf(!IsInitialized(), TEXT("Context is already initialized"));
+		Scene = InScene;
+	}
+
+	void Process();
+
+	static void DestroyProxy(FMassDestroyRenderStateContext* InContext, FPrimitiveSceneProxy* InPrimitiveSceneProxy, FSceneInterface* InScene);
+	static void DestroyProxy(FMassDestroyRenderStateContext* InContext, FDeferredDecalProxy* InDeferredDecalProxy, FSceneInterface* InScene);
+
+private:
+	FSceneInterface* Scene = nullptr;
+	TArray<FPrimitiveSceneProxy*> PrimitiveSceneProxies;
+	TArray<FDeferredDecalProxy*> DeferredDecalProxies;
+};
\ No newline at end of file
diff --git a/Engine/Source/Runtime/Mass/MassEngine/Private/Mesh/MassEngineMeshFragments.cpp b/Engine/Source/Runtime/Mass/MassEngine/Private/Mesh/MassEngineMeshFragments.cpp
--- a/Engine/Source/Runtime/Mass/MassEngine/Private/Mesh/MassEngineMeshFragments.cpp
+++ b/Engine/Source/Runtime/Mass/MassEngine/Private/Mesh/MassEngineMeshFragments.cpp
@@ -2,7 +2,114 @@
 
 #include "Mesh/MassEngineMeshFragments.h"
 
+#include "InstancedStaticMeshSceneProxyDesc.h"
+#include "MassEntityFragments.h"
+#include "StaticMeshSceneProxyDesc.h"
+
+struct FTransformFragment;
+
+//----------------------------------------------------------------------//
+// FMassStaticMeshFragment
+//----------------------------------------------------------------------//
 FMassStaticMeshFragment::FMassStaticMeshFragment(TNotNull<const UStaticMesh*> InMesh)
 	: Mesh(InMesh)
 {
 }
+
+FMassRenderStaticMeshFragment::FMassRenderStaticMeshFragment()
+	: StaticMeshSceneProxyDesc(new FStaticMeshSceneProxyDesc())
+	, bOwnProxyDesc(true)
+{
+}
+
+FMassRenderStaticMeshFragment::FMassRenderStaticMeshFragment(TNotNull<FStaticMeshSceneProxyDesc*> InStaticMeshSceneProxyDesc)
+	: StaticMeshSceneProxyDesc(InStaticMeshSceneProxyDesc)
+	, bOwnProxyDesc(false)
+{
+}
+
+FMassRenderStaticMeshFragment::~FMassRenderStaticMeshFragment()
+{
+	if (bOwnProxyDesc)
+	{
+		delete StaticMeshSceneProxyDesc;
+	}
+}
+
+FMassRenderStaticMeshFragment::FMassRenderStaticMeshFragment(FMassRenderStaticMeshFragment&& Other)
+	: FMassFragment(MoveTemp(Other))
+	, StaticMeshSceneProxyDesc(Other.StaticMeshSceneProxyDesc)
+	, bOwnProxyDesc(Other.bOwnProxyDesc)
+{
+	Other.bOwnProxyDesc = false;
+	Other.StaticMeshSceneProxyDesc = nullptr;
+}
+
+UStaticMesh* FMassRenderStaticMeshFragment::GetStaticMesh()
+{
+	return StaticMeshSceneProxyDesc ? StaticMeshSceneProxyDesc->StaticMesh : nullptr;
+}
+
+const UStaticMesh* FMassRenderStaticMeshFragment::GetStaticMesh() const
+{
+	return StaticMeshSceneProxyDesc ? StaticMeshSceneProxyDesc->StaticMesh : nullptr;
+}
+
+#if WITH_EDITOR
+//----------------------------------------------------------------------//
+// HTypeElementHandleHitProxy
+//----------------------------------------------------------------------//
+IMPLEMENT_HIT_PROXY(HTypeElementHandleHitProxy, HHitProxy);
+
+EMouseCursor::Type HTypeElementHandleHitProxy::GetMouseCursor()
+{
+	return EMouseCursor::Crosshairs;
+}
+#endif // WITH_EDITOR
+
+//----------------------------------------------------------------------//
+// FMassRenderISMFragment
+//----------------------------------------------------------------------//
+FMassRenderISMFragment::FMassRenderISMFragment()
+	: InstancedStaticMeshSceneProxyDesc(new FInstancedStaticMeshSceneProxyDesc())
+	, bOwnProxyDesc(true)
+{
+}
+
+FMassRenderISMFragment::FMassRenderISMFragment(TNotNull<FInstancedStaticMeshSceneProxyDesc*> InInstancedStaticMeshSceneProxyDesc)
+	: InstancedStaticMeshSceneProxyDesc(InInstancedStaticMeshSceneProxyDesc)
+	, bOwnProxyDesc(false)
+{
+}
+
+FMassRenderISMFragment::~FMassRenderISMFragment()
+{
+	if (bOwnProxyDesc)
+	{
+		delete InstancedStaticMeshSceneProxyDesc;
+	}
+}
+
+FMassRenderISMFragment::FMassRenderISMFragment(FMassRenderISMFragment&& Other)
+: FMassFragment(MoveTemp(Other))
+, PerInstanceHitProxy(MoveTemp(Other.PerInstanceHitProxy))
+, PerInstanceSMData(MoveTemp(Other.PerInstanceSMData))
+, PerInstanceSMCustomData(MoveTemp(Other.PerInstanceSMCustomData))
+, AdditionalRandomSeeds(MoveTemp(Other.AdditionalRandomSeeds))
+, InstancedStaticMeshSceneProxyDesc(Other.InstancedStaticMeshSceneProxyDesc)
+, InstancingRandomSeed(Other.InstancingRandomSeed)
+, bOwnProxyDesc(Other.bOwnProxyDesc)
+{
+	Other.bOwnProxyDesc = false;
+	Other.InstancedStaticMeshSceneProxyDesc = nullptr;
+}
+
+UStaticMesh* FMassRenderISMFragment::GetStaticMesh()
+{
+	return InstancedStaticMeshSceneProxyDesc ? InstancedStaticMeshSceneProxyDesc->StaticMesh : nullptr;
+}
+
+const UStaticMesh* FMassRenderISMFragment::GetStaticMesh() const
+{
+	return InstancedStaticMeshSceneProxyDesc ? InstancedStaticMeshSceneProxyDesc->StaticMesh : nullptr;
+}
\ No newline at end of file
diff --git a/Engine/Source/Runtime/Mass/MassEngine/Private/Mesh/MassEngineMeshUtils.cpp b/Engine/Source/Runtime/Mass/MassEngine/Private/Mesh/MassEngineMeshUtils.cpp
new file mode 100644
--- /dev/null
+++ b/Engine/Source/Runtime/Mass/MassEngine/Private/Mesh/MassEngineMeshUtils.cpp
@@ -0,0 +1,150 @@
+// Copyright Epic Games, Inc. All Rights Reserved.
+#include "Mesh/MassEngineMeshUtils.h"
+#include "Mesh/MassEngineMeshFragments.h"
+#include "InstancedStaticMeshComponentHelper.h"
+#include "InstancedStaticMeshSceneProxyDesc.h"
+#include "MassEntityFragments.h"
+#include "StaticMeshSceneProxyDesc.h"
+
+namespace UE::MassEngine::Mesh
+{
+	void InitializePrimitiveFragmentsFromComponent(TNotNull<const UActorComponent*> Component, FTransformFragment& TransformFragment, FMassRenderPrimitiveFragment& RenderPrimitiveFragment)
+	{
+		// Initialize properties not handled by InitializeFromPrimitiveComponent
+		const UPrimitiveComponent* PrimitiveComponent = CastChecked<UPrimitiveComponent>(Component);
+		TransformFragment.SetTransform(PrimitiveComponent->GetComponentToWorld());
+		RenderPrimitiveFragment.bIsVisible = PrimitiveComponent->IsVisible();
+		RenderPrimitiveFragment.bRasterizeAsFilledConvexVolume = PrimitiveComponent->bRasterizeAsFilledConvexVolume;
+		RenderPrimitiveFragment.CustomPrimitiveData = PrimitiveComponent->GetCustomPrimitiveData();
+		RenderPrimitiveFragment.DetailMode = PrimitiveComponent->DetailMode;
+		RenderPrimitiveFragment.RuntimeVirtualTextures = PrimitiveComponent->GetRuntimeVirtualTextures();
+	}
+
+	void InitializeStaticMeshFragmentsFromComponent(TNotNull<const UActorComponent*> Component, FTransformFragment& TransformFragment, FMassRenderPrimitiveFragment& RenderPrimitiveFragment, FMassRenderStaticMeshFragment& RenderStaticMeshFragment)
+	{
+		InitializePrimitiveFragmentsFromComponent(Component, TransformFragment, RenderPrimitiveFragment);
+
+		const UStaticMeshComponent* StaticMeshComponent = CastChecked<const UStaticMeshComponent>(Component);
+		RenderStaticMeshFragment.StaticMeshSceneProxyDesc->InitializeFromStaticMeshComponent(StaticMeshComponent);
+
+		UStaticMesh* StaticMesh = RenderStaticMeshFragment.StaticMeshSceneProxyDesc->StaticMesh;
+		checkf(StaticMesh, TEXT("Expecting a valid static mesh at this point"));
+
+		RenderPrimitiveFragment.LocalBounds = StaticMesh->GetBounds();
+		RenderPrimitiveFragment.WorldBounds = RenderPrimitiveFragment.LocalBounds.TransformBy(TransformFragment.GetTransform());
+	}
+
+	FBoxSphereBounds CalculateInstancedStaticMeshBounds(const FTransformFragment& TransformFragment, const FMassRenderISMFragment& RenderISMFragment, const EBoundsType BoundsType)
+	{
+		UStaticMesh* StaticMesh = RenderISMFragment.InstancedStaticMeshSceneProxyDesc ? RenderISMFragment.InstancedStaticMeshSceneProxyDesc->StaticMesh : nullptr;
+		checkf(StaticMesh, TEXT("Expecting a valid static mesh at this point"));
+
+		if (RenderISMFragment.PerInstanceSMData.Num() > 0)
+		{
+			const bool bWorldSpace = (BoundsType != EBoundsType::LocalBounds);
+			const FBox InstanceBounds = (BoundsType == EBoundsType::NavigationBounds) ? FInstancedStaticMeshComponentHelper::GetInstanceNavigationBounds(RenderISMFragment) : StaticMesh->GetBounds().GetBox();
+			const FMatrix ComponentTransformMatrix = TransformFragment.GetTransform().ToMatrixWithScale();
+			if (InstanceBounds.IsValid)
+			{
+				FBoxSphereBounds::Builder BoundsBuilder;
+				for (const FInstancedStaticMeshInstanceData& InstanceSMData : RenderISMFragment.PerInstanceSMData)
+				{
+					if (bWorldSpace)
+					{
+						BoundsBuilder += InstanceBounds.TransformBy(InstanceSMData.Transform * ComponentTransformMatrix);
+					}
+					else
+					{
+						BoundsBuilder += InstanceBounds.TransformBy(InstanceSMData.Transform);
+					}
+				}
+				return BoundsBuilder;
+			}
+		}
+		return FBoxSphereBounds(FVector::ZeroVector, FVector::ZeroVector, 0.f);
+	}
+
+	void InitializeISMFragmentsFromComponent(TNotNull<const UActorComponent*> Component, FTransformFragment& TransformFragment, FMassRenderPrimitiveFragment& RenderPrimitiveFragment, FMassRenderISMFragment& RenderISMFragment)
+	{
+		InitializePrimitiveFragmentsFromComponent(Component, TransformFragment, RenderPrimitiveFragment);
+
+		if (!RenderISMFragment.InstancedStaticMeshSceneProxyDesc)
+		{
+			RenderISMFragment.InstancedStaticMeshSceneProxyDesc = new FInstancedStaticMeshSceneProxyDesc();
+			RenderISMFragment.bOwnProxyDesc = true;
+		}
+
+		const UInstancedStaticMeshComponent* ISMComponent = CastChecked<const UInstancedStaticMeshComponent>(Component);
+		RenderISMFragment.InstancedStaticMeshSceneProxyDesc->InitializeFromStaticMeshComponent(ISMComponent);
+
+		RenderISMFragment.InstancedStaticMeshSceneProxyDesc->bCollisionEnabled = !!RenderISMFragment.InstancedStaticMeshSceneProxyDesc->bCollisionEnabled && !ISMComponent->bDisableCollision;
+		RenderISMFragment.AdditionalRandomSeeds = ISMComponent->AdditionalRandomSeeds;
+		// @todo optimize the TArray to TSparseArray copy here
+		RenderISMFragment.PerInstanceSMData.Reset();
+		RenderISMFragment.PerInstanceSMData.Reserve(ISMComponent->PerInstanceSMData.Num());
+		for (const FInstancedStaticMeshInstanceData& Data : ISMComponent->PerInstanceSMData)
+		{
+			RenderISMFragment.PerInstanceSMData.Add(Data);
+		}
+		RenderISMFragment.PerInstanceSMCustomData = ISMComponent->PerInstanceSMCustomData;
+		RenderISMFragment.InstancingRandomSeed = ISMComponent->InstancingRandomSeed;
+
+		RenderPrimitiveFragment.LocalBounds = CalculateInstancedStaticMeshBounds(TransformFragment, RenderISMFragment, EBoundsType::LocalBounds);
+		RenderPrimitiveFragment.WorldBounds = CalculateInstancedStaticMeshBounds(TransformFragment, RenderISMFragment, EBoundsType::WorldBounds);
+	}
+
+
+#if WITH_EDITOR
+	void InitializePrimitiveSceneProxyDescFromEditorFragment(const FMassEditorVisualizationMeshFragment& EditorMeshFragment, FPrimitiveSceneProxyDesc& PrimitiveSceneProxyDesc)
+	{
+		PrimitiveSceneProxyDesc.bShouldRenderSelected = EditorMeshFragment.bShouldRenderSelected;
+		PrimitiveSceneProxyDesc.bSelected = EditorMeshFragment.bShouldRenderSelected;
+		PrimitiveSceneProxyDesc.bIndividuallySelected = false;
+		PrimitiveSceneProxyDesc.bIsHiddenEd = EditorMeshFragment.bHiddenInEditor;
+		PrimitiveSceneProxyDesc.bIsVisibleEditor = !EditorMeshFragment.bHiddenInEditor;
+		PrimitiveSceneProxyDesc.bLevelInstanceEditingState = EditorMeshFragment.bLevelInstanceEditingState;
+
+	}
+#endif // WITH_EDITOR
+
+	void InitializeStaticMeshSceneProxyDescFromFragment(const FMassStaticMeshFragment& StaticMeshFragment,	const FMassVisualizationMeshFragment& MeshFragment, const FTransformFragment& TransformFragment, FStaticMeshSceneProxyDesc& StaticMeshSceneProxyDesc)
+	{
+		StaticMeshSceneProxyDesc.StaticMesh = const_cast<UStaticMesh*>(StaticMeshFragment.Mesh.Get());
+
+		//////////////////////////////////////////////////////////////////////////////	
+		// This should maybe be done by the RenderStateHelpers
+		if (FStaticMeshRenderData* RenderData = StaticMeshSceneProxyDesc.StaticMesh ? StaticMeshSceneProxyDesc.StaticMesh->GetRenderData() : nullptr)
+		{
+			StaticMeshSceneProxyDesc.NaniteResources = RenderData->NaniteResourcesPtr.Get();
+		}
+#if WITH_EDITOR
+		StaticMeshSceneProxyDesc.TextureStreamingTransformScale = TransformFragment.GetTransform().GetMaximumAxisScale();
+#endif
+		// This should maybe be done by the RenderStateHelpers
+		//////////////////////////////////////////////////////////////////////////////	
+
+		StaticMeshSceneProxyDesc.TranslucencySortPriority = MeshFragment.TranslucencySortPriority;
+		StaticMeshSceneProxyDesc.CastShadow = MeshFragment.CastShadow;
+		StaticMeshSceneProxyDesc.bCastShadowAsTwoSided = MeshFragment.CastShadowAsTwoSided;
+		StaticMeshSceneProxyDesc.bCastHiddenShadow = MeshFragment.CastHiddenShadow;
+		StaticMeshSceneProxyDesc.bCastFarShadow = MeshFragment.CastFarShadow;
+		StaticMeshSceneProxyDesc.bCastInsetShadow = MeshFragment.CastInsetShadow;
+		StaticMeshSceneProxyDesc.bCastContactShadow = MeshFragment.CastContactShadow;
+		StaticMeshSceneProxyDesc.bAffectDynamicIndirectLighting = MeshFragment.AffectDynamicIndirectLighting;
+		StaticMeshSceneProxyDesc.bAffectIndirectLightingWhileHidden = MeshFragment.AffectIndirectLightingWhileHidden;
+		StaticMeshSceneProxyDesc.bReceivesDecals = MeshFragment.ReceivesDecals;
+		StaticMeshSceneProxyDesc.bRenderCustomDepth = MeshFragment.RenderCustomDepth;
+		StaticMeshSceneProxyDesc.CustomDepthStencilValue = (int32)MeshFragment.CustomDepthStencilValue;
+
+		if (MeshFragment.NeverDistanceCull)
+		{
+			StaticMeshSceneProxyDesc.CachedMaxDrawDistance = 0.0f;
+		}
+	}
+
+	void InitializeInstanceStaticMeshSceneProxyDescFromFragment(const FMassStaticMeshFragment& StaticMeshFragment, const FMassVisualizationMeshFragment& MeshFragment, const FTransformFragment& TransformFragment, FInstancedStaticMeshSceneProxyDesc& ISMSceneProxyDesc)
+	{
+		InitializeStaticMeshSceneProxyDescFromFragment(StaticMeshFragment, MeshFragment, TransformFragment, ISMSceneProxyDesc);
+	};
+
+}
\ No newline at end of file
diff --git a/Engine/Source/Runtime/Mass/MassEngine/Private/Mesh/MassEngineRenderISMProcessors.cpp b/Engine/Source/Runtime/Mass/MassEngine/Private/Mesh/MassEngineRenderISMProcessors.cpp
new file mode 100644
--- /dev/null
+++ b/Engine/Source/Runtime/Mass/MassEngine/Private/Mesh/MassEngineRenderISMProcessors.cpp
@@ -0,0 +1,471 @@
+// Copyright Epic Games, Inc. All Rights Reserved.
+
+#include "Mesh/MassEngineRenderISMProcessors.h"
+
+#include "MassDestroyRenderStateContext.h"
+#include "MassEngineTypes.h"
+#include "MassEntityFragments.h"
+#include "MassExecutionContext.h"
+#include "MassISMRenderStateHelper.h"
+#include "MassRenderStateHelper.h"
+#include "MassSignalSubsystem.h"
+#include "MassStaticMeshRenderStateHelper.h"
+#include "Mesh/MassEngineMeshFragments.h"
+#include "Mesh/MassEngineMeshUtils.h"
+#include "Misc/HashBuilder.h"
+#include "StaticMeshComponentHelper.h"
+
+#include UE_INLINE_GENERATED_CPP_BY_NAME(MassEngineRenderISMProcessors)
+
+//----------------------------------------------------------------------//
+// UMassInstantiateStaticMeshAndCreateISMRenderStateProcessor
+//----------------------------------------------------------------------//
+void UMassInstantiateStaticMeshAndCreateISMRenderStateProcessor::ConfigureQueries(const TSharedRef<FMassEntityManager>& EntityManager)
+{
+	Super::ConfigureQueries(EntityManager);
+ 	EntityQuery.AddRequirement<FMassTypeElementFragment>(EMassFragmentAccess::ReadOnly, EMassFragmentPresence::Optional);
+	EntityQuery.AddTagRequirement<FMassRenderISMCandidateTag>(EMassFragmentPresence::All);
+	EntityQuery.AddConstSharedRequirement<FMassRenderISMLinkFragment>(EMassFragmentPresence::None);
+	EntityQuery.AddRequirement<FMassRenderISMInstanceFragment>(EMassFragmentAccess::ReadOnly, EMassFragmentPresence::None);
+
+	ISMEntityComposition.Add<FTransformFragment>();
+	ISMEntityComposition.Add<FMassRenderStateFragment>();
+	ISMEntityComposition.Add<FMassRenderPrimitiveFragment>();
+	ISMEntityComposition.Add<FMassRenderISMFragment>();
+	ISMEntityArchetype = EntityManager->CreateArchetype(ISMEntityComposition);
+	ProcessorEntityCreationRequirements.AddCreatedArchetype(ISMEntityComposition);
+
+	ISMEntityCompositionWithMaterialOverrides = ISMEntityComposition;
+	ISMEntityCompositionWithMaterialOverrides.Add<FMassOverrideMaterialsFragment>();
+	ISMEntityArchetypeWithMaterialOverrides = EntityManager->CreateArchetype(ISMEntityCompositionWithMaterialOverrides);
+	ProcessorEntityCreationRequirements.AddCreatedArchetype(ISMEntityCompositionWithMaterialOverrides);
+
+#if WITH_EDITOR
+	ISMEntityCompositionWithEditorMesh = ISMEntityComposition;
+	ISMEntityCompositionWithEditorMesh.Add<FMassEditorVisualizationMeshFragment>();
+	ISMEntityArchetypeWithEditorMesh = EntityManager->CreateArchetype(ISMEntityCompositionWithEditorMesh);
+	ProcessorEntityCreationRequirements.AddCreatedArchetype(ISMEntityCompositionWithEditorMesh);
+
+	ISMEntityCompositionWithMaterialOverridesAndEditorMesh = ISMEntityCompositionWithMaterialOverrides;
+	ISMEntityCompositionWithMaterialOverridesAndEditorMesh.Add<FMassEditorVisualizationMeshFragment>();
+	ISMEntityArchetypeWithMaterialOverridesAndEditorMesh = EntityManager->CreateArchetype(ISMEntityCompositionWithMaterialOverridesAndEditorMesh);
+	ProcessorEntityCreationRequirements.AddCreatedArchetype(ISMEntityCompositionWithMaterialOverridesAndEditorMesh);
+#endif // WITH_EDITOR
+
+}
+
+void UMassInstantiateStaticMeshAndCreateISMRenderStateProcessor::Execute(FMassEntityManager& EntityManager, FMassExecutionContext& ExecutionContext)
+{
+	struct FISMEntityInfo
+	{
+		uint32 SourceSharedFragmentHash = 0;
+		const FMassStaticMeshFragment& StaticMeshFragment;
+		const FMassVisualizationMeshFragment& MeshFragment;
+		uint32 InstanceCount = 0;
+		uint32 CurrentInstance = 0;
+	};
+	struct FISMEntityBatch
+	{
+		const FMassOverrideMaterialsFragment* OverrideMaterialsFragment = nullptr;
+		const FMassEditorVisualizationMeshFragment* EditorMeshFragment = nullptr;
+		TArray<FISMEntityInfo> EntitiesInfo;
+		TArray<FMassEntityHandle> CreatedEntities;
+	};
+
+	TArray<FISMEntityBatch> ISMEntityBatches;
+
+	auto RetrieveISMEntityInfo = [&ISMEntityBatches](
+		const FMassStaticMeshFragment& StaticMeshFragment,
+		const FMassVisualizationMeshFragment& MeshFragment,
+		const FMassOverrideMaterialsFragment* OverrideMaterialsFragment,
+		const FMassEditorVisualizationMeshFragment* EditorMeshFragment,
+		bool bCreateIfNotFound,
+		FMassEntityHandle* ISMEntityHandle = nullptr) -> FISMEntityInfo*
+	{
+		FISMEntityBatch* Batch = ISMEntityBatches.FindByPredicate(
+		[OverrideMaterialsFragment, EditorMeshFragment](const FISMEntityBatch& Batch)
+			{
+				return Batch.OverrideMaterialsFragment == OverrideMaterialsFragment &&
+					Batch.EditorMeshFragment == EditorMeshFragment;
+			});
+		if (!Batch)
+		{
+			if (!bCreateIfNotFound)
+			{
+				return nullptr;
+			}
+			Batch = &ISMEntityBatches.Emplace_GetRef(OverrideMaterialsFragment, EditorMeshFragment);
+		}
+
+		// Build the source shared fragment hash
+		FHashBuilder SourceSharedFragmentsHash;
+		auto AddFragmentToHash = [&SourceSharedFragmentsHash](const void* FragmentPtr)
+			{
+				SourceSharedFragmentsHash.AppendRaw(&FragmentPtr, sizeof(&FragmentPtr));
+			};
+		AddFragmentToHash(&StaticMeshFragment);
+		AddFragmentToHash(&MeshFragment);
+
+		// Find if the ISM entity already exist, if not add it
+		FISMEntityInfo* ISMEntity;
+		int32 Index = Batch->EntitiesInfo.IndexOfByPredicate([Hash = SourceSharedFragmentsHash.GetHash()](const FISMEntityInfo& ISMEntityInfo) { return ISMEntityInfo.SourceSharedFragmentHash == Hash; });
+		if (Index == INDEX_NONE)
+		{
+			if (!bCreateIfNotFound)
+			{
+				return nullptr;
+			}
+			Index = Batch->EntitiesInfo.Num();
+			ISMEntity = &Batch->EntitiesInfo.Emplace_GetRef(SourceSharedFragmentsHash.GetHash(), StaticMeshFragment, MeshFragment);
+		}
+		else
+		{
+			ISMEntity = &Batch->EntitiesInfo[Index];
+		}
+
+		if (ISMEntityHandle)
+		{
+			checkf(Index < Batch->CreatedEntities.Num(), TEXT("Expecting all the entities to be created at this point"));
+			*ISMEntityHandle = Batch->CreatedEntities[Index];
+		}
+		return ISMEntity;
+	};
+
+	//////////////////////////////////////////////////////////////////////////////
+	// First gather all the different ISM entities we need to create.
+	EntityQuery.ForEachEntityChunk(ExecutionContext,
+		[&RetrieveISMEntityInfo](FMassExecutionContext& Context)
+		{
+			const FMassStaticMeshFragment& StaticMeshFragment = Context.GetConstSharedFragment<FMassStaticMeshFragment>();
+			const FMassVisualizationMeshFragment& MeshFragment = Context.GetConstSharedFragment<FMassVisualizationMeshFragment>();
+			const FMassOverrideMaterialsFragment* OverrideMaterialsFragment = Context.GetConstSharedFragmentPtr<FMassOverrideMaterialsFragment>();
+#if WITH_EDITOR
+			const FMassEditorVisualizationMeshFragment* EditorMeshFragment = Context.GetConstSharedFragmentPtr<FMassEditorVisualizationMeshFragment>();
+#else
+			const FMassEditorVisualizationMeshFragment* EditorMeshFragment = nullptr;
+#endif // WITH_EDITOR
+
+			// Find the batch matching the shared fragment
+			FISMEntityInfo* ISMEntityInfo = RetrieveISMEntityInfo(StaticMeshFragment, MeshFragment, OverrideMaterialsFragment, EditorMeshFragment, /*bCreateIfNotFound*/true);
+			checkf(ISMEntityInfo, TEXT("Expecting the ISM to be created"));
+			ISMEntityInfo->InstanceCount += Context.GetNumEntities();
+		});
+
+	//////////////////////////////////////////////////////////////////////////////
+	// Batch create the mass entities
+	for (FISMEntityBatch& Batch : ISMEntityBatches)
+	{
+		FMassArchetypeHandle ArchetypeHandle;
+		FMassArchetypeSharedFragmentValues SharedFragmentValues;
+		if (!Batch.OverrideMaterialsFragment && !Batch.EditorMeshFragment)
+		{
+			ArchetypeHandle = ISMEntityArchetype;
+		}
+		else if (Batch.OverrideMaterialsFragment && !Batch.EditorMeshFragment)
+		{
+			SharedFragmentValues.Add(FConstSharedStruct::Make(*Batch.OverrideMaterialsFragment));
+			ArchetypeHandle = ISMEntityArchetypeWithMaterialOverrides;
+		}
+#if WITH_EDITOR
+		else if (!Batch.OverrideMaterialsFragment && Batch.EditorMeshFragment)
+		{
+			SharedFragmentValues.Add(FConstSharedStruct::Make(*Batch.EditorMeshFragment));
+			ArchetypeHandle = ISMEntityArchetypeWithEditorMesh;
+		}
+		else if (Batch.OverrideMaterialsFragment && Batch.EditorMeshFragment)
+		{
+			SharedFragmentValues.Add(FConstSharedStruct::Make(*Batch.OverrideMaterialsFragment));
+			SharedFragmentValues.Add(FConstSharedStruct::Make(*Batch.EditorMeshFragment));
+			ArchetypeHandle = ISMEntityArchetypeWithMaterialOverridesAndEditorMesh;
+		}
+#endif // WITH_EDITOR
+
+		ExecutionContext.BatchCreateEntities(ArchetypeHandle, SharedFragmentValues, Batch.EntitiesInfo.Num(), Batch.CreatedEntities);
+	}
+
+	//////////////////////////////////////////////////////////////////////////////
+	// Gather all the transforms for each ISM entity and linked them to the created ISM entity
+	EntityQuery.ForEachEntityChunk(ExecutionContext,
+		[&EntityManager, &RetrieveISMEntityInfo](FMassExecutionContext& Context)
+		{
+			const FMassStaticMeshFragment& StaticMeshFragment = Context.GetConstSharedFragment<FMassStaticMeshFragment>();
+			const FMassVisualizationMeshFragment& MeshFragment = Context.GetConstSharedFragment<FMassVisualizationMeshFragment>();
+			const FMassOverrideMaterialsFragment* OverrideMaterialsFragment = Context.GetConstSharedFragmentPtr<FMassOverrideMaterialsFragment>();
+			const TConstArrayView<FTransformFragment> TransformFragments = Context.GetFragmentView<FTransformFragment>();
+			const TConstArrayView<FMassTypeElementFragment> TypeElementFragments = Context.GetFragmentView<FMassTypeElementFragment>();
+#if WITH_EDITOR
+			const FMassEditorVisualizationMeshFragment* EditorMeshFragment = Context.GetConstSharedFragmentPtr<FMassEditorVisualizationMeshFragment>();
+#else
+			const FMassEditorVisualizationMeshFragment* EditorMeshFragment = nullptr;
+#endif // WITH_EDITOR
+
+			FMassEntityHandle ISMEntityHandle;
+			FISMEntityInfo* ISMEntityInfo = RetrieveISMEntityInfo(StaticMeshFragment, MeshFragment, OverrideMaterialsFragment, EditorMeshFragment, /*bCreateIfNotFound*/false, &ISMEntityHandle);
+			checkf(ISMEntityInfo, TEXT("Expecting the ISM to be created"));
+			FMassEntityView ISMEntityView(EntityManager, ISMEntityHandle);
+
+			FMassRenderISMLinkFragment LinkedISMEntityFragmentValues;
+			LinkedISMEntityFragmentValues.LinkedEntityHandle = ISMEntityHandle;
+			const FConstSharedStruct LinkedISMEntityFragment = EntityManager.GetOrCreateConstSharedFragment(LinkedISMEntityFragmentValues);
+
+			for (FMassExecutionContext::FEntityIterator EntityIt = Context.CreateEntityIterator(); EntityIt; ++EntityIt)
+			{
+				FMassRenderISMFragment& ISMDescFragment = ISMEntityView.GetFragmentData<FMassRenderISMFragment>();
+				if (ISMDescFragment.PerInstanceSMData.Num() == 0)
+				{
+					ISMDescFragment.PerInstanceSMData.Reset();
+					ISMDescFragment.PerInstanceSMData.Reserve(ISMEntityInfo->InstanceCount);
+#if WITH_EDITOR
+					ISMDescFragment.PerInstanceHitProxy.Reset();
+					ISMDescFragment.PerInstanceHitProxy.Reserve(ISMEntityInfo->InstanceCount);
+#endif // WITH_EDITOR
+				}
+				FMassRenderISMInstanceFragment ISMInfoFragment;
+				ISMInfoFragment.InstanceIndex = ISMDescFragment.PerInstanceSMData.Add(TransformFragments[*EntityIt].GetTransform().ToMatrixWithScale());
+#if WITH_EDITOR
+				ISMInfoFragment.HitProxyIndex = ISMDescFragment.PerInstanceHitProxy.Add( new HTypeElementHandleHitProxy(TypeElementFragments.Num()? TypeElementFragments[*EntityIt].TypeElementHandle : FTypedElementHandle{}));
+#endif // WITH_EDITOR
+
+				// Add new fragments @todo find a way to move to new archetype in one shot
+				FMassEntityHandle EntityHandle = Context.GetEntity(*EntityIt);
+				Context.Defer().PushCommand<FMassCommandAddFragmentInstances>(EntityHandle, MoveTemp(ISMInfoFragment));
+				Context.Defer().PushCommand<FMassDeferredAddCommand>([EntityHandle, LinkedISMEntityFragment](FMassEntityManager& InEntityManager)
+					{
+						if (InEntityManager.IsEntityValid(EntityHandle))
+						{
+							InEntityManager.AddConstSharedFragmentToEntity(EntityHandle, LinkedISMEntityFragment);
+						}
+						else
+						{
+							UE_LOG(LogMass, Log, TEXT("FMassRenderISMLinkFragment shared fragment not added: Entity is no longer valid"));
+						}
+					});
+
+				ISMEntityInfo->CurrentInstance++;
+			}
+		});
+
+	//////////////////////////////////////////////////////////////////////////////
+	// Finished the creation of the ISM entities
+	for (const FISMEntityBatch& Batch : ISMEntityBatches)
+	{
+		for (int32 i = 0; i < Batch.EntitiesInfo.Num(); ++i)
+		{
+			const FISMEntityInfo& ISMInfo = Batch.EntitiesInfo[i];
+			const FMassEntityHandle ISMEntityHandle = Batch.CreatedEntities[i];
+			const FMassEntityView ISMEntityView(EntityManager, ISMEntityHandle);
+
+			FTransformFragment& TransformFragment = ISMEntityView.GetFragmentData<FTransformFragment>();
+			FMassRenderPrimitiveFragment& RenderPrimitiveFragment = ISMEntityView.GetFragmentData<FMassRenderPrimitiveFragment>();
+			FMassRenderISMFragment& ISMDescFragment = ISMEntityView.GetFragmentData<FMassRenderISMFragment>();
+			FMassRenderStateFragment& RenderStateFragment = ISMEntityView.GetFragmentData<FMassRenderStateFragment>();
+
+
+			checkf(ISMDescFragment.InstancedStaticMeshSceneProxyDesc, TEXT("Expecting a valid proxy desc"));
+			UE::MassEngine::Mesh::InitializeInstanceStaticMeshSceneProxyDescFromFragment(ISMInfo.StaticMeshFragment, ISMInfo.MeshFragment, TransformFragment, *ISMDescFragment.InstancedStaticMeshSceneProxyDesc);
+#if WITH_EDITOR
+			if (Batch.EditorMeshFragment)
+			{
+				UE::MassEngine::Mesh::InitializePrimitiveSceneProxyDescFromEditorFragment(*Batch.EditorMeshFragment, *ISMDescFragment.InstancedStaticMeshSceneProxyDesc);
+			}
+#endif // WITH_EDITOR
+
+			RenderPrimitiveFragment.bIsVisible = true;
+			RenderPrimitiveFragment.LocalBounds = UE::MassEngine::Mesh::CalculateInstancedStaticMeshBounds(TransformFragment, ISMDescFragment, UE::MassEngine::Mesh::EBoundsType::LocalBounds);
+			RenderPrimitiveFragment.WorldBounds = UE::MassEngine::Mesh::CalculateInstancedStaticMeshBounds(TransformFragment, ISMDescFragment, UE::MassEngine::Mesh::EBoundsType::WorldBounds);
+
+			RenderStateFragment.CreateRenderStateHelper<FMassISMRenderStateHelper>(ISMEntityHandle, &EntityManager, RenderPrimitiveFragment, Batch.OverrideMaterialsFragment, ISMDescFragment);
+		}
+	}
+}
+
+//----------------------------------------------------------------------//
+// UMassDirtyRenderISMInstanceProcessor
+//----------------------------------------------------------------------//
+UMassDirtyRenderISMInstanceProcessor::UMassDirtyRenderISMInstanceProcessor(const FObjectInitializer& ObjectInitializer)
+	: Super(ObjectInitializer)
+{
+	ExecutionFlags = static_cast<int32>(EProcessorExecutionFlags::AllWorldModes);
+	ExecutionOrder.ExecuteInGroup = UE::Mass::ProcessorGroupNames::SetupRenderState;
+}
+
+void UMassDirtyRenderISMInstanceProcessor::InitializeInternal(UObject& Owner, const TSharedRef<FMassEntityManager>& EntityManager)
+{
+	Super::InitializeInternal(Owner, EntityManager);
+
+	if (const UWorld* World = GetWorld())
+	{
+		if (UMassSignalSubsystem* SignalSubsystem = World->GetSubsystem<UMassSignalSubsystem>())
+		{
+			SubscribeToSignal(*SignalSubsystem, UE::Mass::Signals::TransformChanged);
+#if WITH_EDITOR
+			SubscribeToSignal(*SignalSubsystem, UE::Mass::Signals::SelectionChanged);
+#endif // WITH_EDITOR
+		}
+	}
+}
+
+void UMassDirtyRenderISMInstanceProcessor::ConfigureQueries(const TSharedRef<FMassEntityManager>& EntityManager)
+{
+	EntityQuery.AddRequirement<FTransformFragment>(EMassFragmentAccess::ReadOnly);
+	EntityQuery.AddRequirement<FMassRenderISMInstanceFragment>(EMassFragmentAccess::ReadOnly);
+	EntityQuery.AddConstSharedRequirement<FMassRenderISMLinkFragment>();
+	EntityQuery.AddLinkedEntityRequirement<FMassRenderISMFragment, FMassRenderISMLinkFragment>(EMassFragmentAccess::ReadWrite);
+	EntityQuery.AddLinkedEntityRequirement<FMassRenderStateFragment, FMassRenderISMLinkFragment>(EMassFragmentAccess::ReadWrite);
+}
+
+void UMassDirtyRenderISMInstanceProcessor::SignalEntities(FMassEntityManager& EntityManager, FMassExecutionContext& ExecutionContext, FMassSignalNameLookup& EntitySignals)
+{
+	const UObject* EntityManagerOwner = ExecutionContext.GetEntityManagerChecked().GetOwner();
+	check(EntityManagerOwner);
+
+	TArray<FName> SignalsForEntity;
+	EntityQuery.ForEachEntityChunk(ExecutionContext,
+		[&EntitySignals, &SignalsForEntity](FMassExecutionContext& Context)
+		{
+			const TConstArrayView<FMassRenderISMInstanceFragment> ISMInfoFragments = Context.GetFragmentView<FMassRenderISMInstanceFragment>();
+			const FMassRenderISMLinkFragment LinkedISMEntityFragment = Context.GetConstSharedFragment<FMassRenderISMLinkFragment>();
+			FMassRenderISMFragment* ISMDescFragment = Context.GetIndirectFragmentPtr<FMassRenderISMFragment>(LinkedISMEntityFragment.LinkedEntityHandle);
+			checkf(ISMDescFragment, TEXT("Expecting a ISM desc fragment"));
+			FMassRenderStateFragment* RenderStateFragment = Context.GetIndirectFragmentPtr<FMassRenderStateFragment>(LinkedISMEntityFragment.LinkedEntityHandle);
+			checkf(RenderStateFragment, TEXT("Expecting a render state fragment"));
+			const TConstArrayView<FTransformFragment> TransformFragments = Context.GetFragmentView<FTransformFragment>();
+
+			for (FMassExecutionContext::FEntityIterator EntityIt = Context.CreateEntityIterator(); EntityIt; ++EntityIt)
+			{
+				const int32 InstanceIndex = ISMInfoFragments[*EntityIt].InstanceIndex;
+				checkf(ISMDescFragment->PerInstanceSMData.IsValidIndex(InstanceIndex), TEXT("The entity is out of sync with its instance"));
+#if WITH_EDITOR
+				const int32 HitProxyIndex = ISMInfoFragments[*EntityIt].HitProxyIndex;
+				checkf(ISMDescFragment->PerInstanceHitProxy.IsValidIndex(HitProxyIndex), TEXT("The entity is out of sync with its instance type element"));
+				checkf(ISMDescFragment->PerInstanceSMData.Num() == ISMDescFragment->PerInstanceHitProxy.Num(), TEXT("The ISM entity is out of sync"));
+#endif // WITH_EDITOR
+
+				SignalsForEntity.Reset();
+				EntitySignals.GetSignalsForEntity(Context.GetEntity(*EntityIt), SignalsForEntity);
+
+				const bool bTransformChanged = SignalsForEntity.Contains(UE::Mass::Signals::TransformChanged);
+				if (SignalsForEntity.Num() == 1 && bTransformChanged)
+				{
+					ISMDescFragment->PerInstanceSMData[InstanceIndex] = TransformFragments[*EntityIt].GetTransform().ToMatrixWithScale();
+					// @TODO optimize this to use updatable ISM Scene proxy, but for now Trigger an ISM recreation
+					RenderStateFragment->GetRenderStateHelper().DestroyRenderState(nullptr);
+					Context.Defer().RemoveTag<FMassSceneProxyCreatedTag>(LinkedISMEntityFragment.LinkedEntityHandle);
+				}
+				else
+				{
+					FMassEntityHandle EntityHandle = Context.GetEntity(*EntityIt);
+					Context.Defer().RemoveFragment<FMassRenderISMInstanceFragment>(EntityHandle);
+					Context.Defer().PushCommand<FMassDeferredAddCommand>([EntityHandle](FMassEntityManager& InEntityManager)
+						{
+							if (InEntityManager.IsEntityValid(EntityHandle))
+							{
+								InEntityManager.RemoveConstSharedFragmentFromEntity(EntityHandle, FMassRenderISMLinkFragment::StaticStruct());
+							}
+							else
+							{
+								UE_LOG(LogMass, Log, TEXT("FMassRenderISMLinkFragment shared fragment not added: Entity is no longer valid"));
+							}
+						});
+
+					// Delete the ism entity if this is the last instance in the desc
+					if (ISMDescFragment->PerInstanceSMData.Num() == 1)
+					{
+						Context.Defer().DestroyEntity(LinkedISMEntityFragment.LinkedEntityHandle);
+					}
+					else
+					{
+						// Remove the transform of that instance and trigger a SetupRenderState
+						ISMDescFragment->PerInstanceSMData.RemoveAt(InstanceIndex);
+#if WITH_EDITOR
+						ISMDescFragment->PerInstanceHitProxy.RemoveAt(HitProxyIndex);
+#endif // WITH_EDITOR
+						// @TODO optimize this to use updatable ISM Scene proxy, but for now Trigger an ISM recreation
+						RenderStateFragment->GetRenderStateHelper().DestroyRenderState(nullptr);
+						Context.Defer().RemoveTag<FMassSceneProxyCreatedTag>(LinkedISMEntityFragment.LinkedEntityHandle);
+					}
+				}
+			}
+		});
+}
+
+// ----------------------------------------------------------------------//
+// UMassRemoveInstanceAndDestroyISMRenderStateProcessor
+//----------------------------------------------------------------------//
+UMassRemoveInstanceAndDestroyISMRenderStateProcessor::UMassRemoveInstanceAndDestroyISMRenderStateProcessor()
+	: EntityQuery(*this)
+{
+	ExecutionFlags = static_cast<int32>(EProcessorExecutionFlags::AllWorldModes);
+	ObservedType = FMassRenderISMInstanceFragment::StaticStruct();
+	ObservedOperations = EMassObservedOperationFlags::Remove;
+	bRequiresGameThreadExecution = true; // @todo revisit if the execution is required on game thread
+}
+
+void UMassRemoveInstanceAndDestroyISMRenderStateProcessor::ConfigureQueries(const TSharedRef<FMassEntityManager>& EntityManager)
+{
+	EntityQuery.AddRequirement<FMassRenderISMInstanceFragment>(EMassFragmentAccess::ReadOnly);
+	EntityQuery.AddConstSharedRequirement<FMassRenderISMLinkFragment>();
+	EntityQuery.AddLinkedEntityRequirement<FMassRenderISMFragment, FMassRenderISMLinkFragment>(EMassFragmentAccess::ReadOnly);
+	EntityQuery.AddLinkedEntityRequirement<FMassRenderStateFragment, FMassRenderISMLinkFragment>(EMassFragmentAccess::ReadOnly);
+}
+
+void UMassRemoveInstanceAndDestroyISMRenderStateProcessor::Execute(FMassEntityManager& EntityManager, FMassExecutionContext& ExecutionContext)
+{
+	FMassDestroyRenderStateContext BatchingContext(nullptr);
+
+	TArray<FMassEntityHandle> EntitiesToDestroy;
+
+	EntityQuery.ForEachEntityChunk(ExecutionContext,
+		[&BatchingContext, &EntitiesToDestroy](FMassExecutionContext& Context)
+		{
+			const TConstArrayView<FMassRenderISMInstanceFragment> ISMInfoFragments = Context.GetFragmentView<FMassRenderISMInstanceFragment>();
+			const FMassRenderISMLinkFragment LinkedISMEntityFragment = Context.GetConstSharedFragment<FMassRenderISMLinkFragment>();
+			FMassRenderISMFragment* ISMDescFragment = Context.GetIndirectFragmentPtr<FMassRenderISMFragment>(LinkedISMEntityFragment.LinkedEntityHandle);
+			checkf(ISMDescFragment, TEXT("Expecting a ISM desc fragment"));
+			FMassRenderStateFragment* RenderStateFragment = Context.GetIndirectFragmentPtr<FMassRenderStateFragment>(LinkedISMEntityFragment.LinkedEntityHandle);
+			checkf(RenderStateFragment, TEXT("Expecting a render state fragment"));
+
+			for (FMassExecutionContext::FEntityIterator EntityIt = Context.CreateEntityIterator(); EntityIt; ++EntityIt)
+			{
+				const int32 InstanceIndex = ISMInfoFragments[*EntityIt].InstanceIndex;
+				checkf(ISMDescFragment->PerInstanceSMData.IsValidIndex(InstanceIndex), TEXT("The entity is out of sync with its instance"));
+#if WITH_EDITOR
+				const int32 HitProxyIndex = ISMInfoFragments[*EntityIt].HitProxyIndex;
+				checkf(ISMDescFragment->PerInstanceHitProxy.IsValidIndex(HitProxyIndex), TEXT("The entity is out of sync with its instance type element"));
+				checkf(ISMDescFragment->PerInstanceSMData.Num() == ISMDescFragment->PerInstanceHitProxy.Num(), TEXT("The ISM entity is out of sync"))
+#endif // WITH_EDITOR
+
+				// Whether we remove one instance or it is the last one, let destroy the render state.
+				if (!BatchingContext.IsInitialized())
+				{
+					BatchingContext.Initialize(RenderStateFragment->GetRenderStateHelper().GetScene());
+				}
+				RenderStateFragment->GetRenderStateHelper().DestroyRenderState(&BatchingContext);
+
+				// Delete the ism entity if this is the last instance in the desc
+				if (ISMDescFragment->PerInstanceSMData.Num() == 1)
+				{
+					EntitiesToDestroy.Add(LinkedISMEntityFragment.LinkedEntityHandle);
+				}
+				else
+				{
+					// Remove the transform of that instance and trigger a CreateRenderState
+					ISMDescFragment->PerInstanceSMData.RemoveAt(InstanceIndex);
+#if WITH_EDITOR
+					ISMDescFragment->PerInstanceHitProxy.RemoveAt(HitProxyIndex);
+#endif // WITH_EDITOR
+					Context.Defer().RemoveTag<FMassSceneProxyCreatedTag>(LinkedISMEntityFragment.LinkedEntityHandle);
+				}
+			}
+		});
+
+	if (BatchingContext.IsInitialized())
+	{
+		BatchingContext.Process();
+	}
+
+	if (EntitiesToDestroy.Num())
+	{
+		ExecutionContext.Defer().DestroyEntities(EntitiesToDestroy);
+	}
+}
diff --git a/Engine/Source/Runtime/Mass/MassEngine/Private/Mesh/MassEngineRenderStateProcessors.cpp b/Engine/Source/Runtime/Mass/MassEngine/Private/Mesh/MassEngineRenderStateProcessors.cpp
new file mode 100644
--- /dev/null
+++ b/Engine/Source/Runtime/Mass/MassEngine/Private/Mesh/MassEngineRenderStateProcessors.cpp
@@ -0,0 +1,208 @@
+// Copyright Epic Games, Inc. All Rights Reserved.
+
+#include "Mesh/MassEngineRenderStateProcessors.h"
+
+#include "AI/MassEngineNavigationProcessors.h"
+#include "MassDestroyRenderStateContext.h"
+#include "MassEngineTypes.h"
+#include "MassExecutionContext.h"
+#include "MassISMRenderStateHelper.h"
+#include "MassRenderStateHelper.h"
+#include "MassSignalSubsystem.h"
+#include "MassStaticMeshRenderStateHelper.h"
+#include "Mesh/MassEngineMeshFragments.h"
+#include "StaticMeshComponentHelper.h"
+#include "StaticMeshSceneProxyDesc.h"
+
+#include UE_INLINE_GENERATED_CPP_BY_NAME(MassEngineRenderStateProcessors)
+
+// ----------------------------------------------------------------------//
+// UMassCreateRenderStateProcessor
+//----------------------------------------------------------------------//
+UMassCreateRenderStateProcessor::UMassCreateRenderStateProcessor(const FObjectInitializer& ObjectInitializer)
+	: Super(ObjectInitializer)
+	, EntityQuery(*this)
+{
+	bRequiresGameThreadExecution = true; // @todo revisit if the execution is required on game thread
+	ExecutionFlags = static_cast<int32>(EProcessorExecutionFlags::AllWorldModes);
+	ExecutionOrder.ExecuteInGroup = UE::Mass::ProcessorGroupNames::CreateRenderState;
+	ExecutionOrder.ExecuteAfter.Add(UE::Mass::ProcessorGroupNames::SetupRenderState);
+}
+
+void UMassCreateRenderStateProcessor::ConfigureQueries(const TSharedRef<FMassEntityManager>& EntityManager)
+{
+	EntityQuery.AddRequirement<FMassRenderStateFragment>(EMassFragmentAccess::ReadWrite);
+	EntityQuery.AddTagRequirement<FMassSceneProxyCreatedTag>(EMassFragmentPresence::None);
+}
+
+void UMassCreateRenderStateProcessor::Execute(FMassEntityManager& EntityManager, FMassExecutionContext& ExecutionContext)
+{
+	FRegisterComponentContext BatchingContext(nullptr);
+
+	EntityQuery.ForEachEntityChunk(ExecutionContext,
+		[&BatchingContext](FMassExecutionContext& Context)
+		{
+			const TArrayView<FMassRenderStateFragment>& RenderStateFragments = Context.GetMutableFragmentView<FMassRenderStateFragment>();
+
+			for (FMassExecutionContext::FEntityIterator EntityIt = Context.CreateEntityIterator(); EntityIt; ++EntityIt)
+			{
+				FMassPrimitiveRenderStateHelper& Helper = RenderStateFragments[*EntityIt].GetRenderStateHelper<FMassPrimitiveRenderStateHelper>();
+				if (!Helper.IsRenderStateCreated() && Helper.ShouldCreateRenderState())
+				{
+					if (!BatchingContext.IsInitialized())
+					{
+						BatchingContext.Initialize(Helper.GetWorld());
+					}
+					Helper.CreateRenderState(&BatchingContext);
+				}
+				else
+				{
+					Context.Defer().AddTag<FMassSceneProxyCreatedTag>(Context.GetEntity(*EntityIt));
+				}
+			}
+		});
+
+	if (BatchingContext.IsInitialized())
+	{
+		BatchingContext.Process();
+#if WITH_EDITOR
+		// Make sure any new created hit proxies are taken into account by the ProxyMap
+		if (GEditor)
+		{
+			GEditor->RedrawLevelEditingViewports(/*bInvalidateProxies*/true);
+		}
+#endif // WITH_EDITOR
+	}
+}
+
+// ----------------------------------------------------------------------//
+// UMassRenderStateDirtyUpdateProcessor
+//----------------------------------------------------------------------//
+UMassRenderStateDirtyUpdateProcessor::UMassRenderStateDirtyUpdateProcessor(const FObjectInitializer & ObjectInitializer)
+	: Super(ObjectInitializer)
+{
+	bRequiresGameThreadExecution = true; // @todo revisit if the execution is required on game thread
+	ExecutionFlags = static_cast<int32>(EProcessorExecutionFlags::AllWorldModes);
+	ExecutionOrder.ExecuteInGroup = UE::Mass::ProcessorGroupNames::CreateRenderState;
+	ExecutionOrder.ExecuteAfter.Add(UE::Mass::ProcessorGroupNames::SetupRenderState);
+}
+
+void UMassRenderStateDirtyUpdateProcessor::InitializeInternal(UObject& Owner, const TSharedRef<FMassEntityManager>& EntityManager)
+{
+	Super::InitializeInternal(Owner, EntityManager);
+
+	if (const UWorld* World = GetWorld())
+	{
+		if (UMassSignalSubsystem* SignalSubsystem = World->GetSubsystem<UMassSignalSubsystem>())
+		{
+			SubscribeToSignal(*SignalSubsystem, UE::Mass::Signals::TransformChanged);
+			SubscribeToSignal(*SignalSubsystem, UE::Mass::Signals::MeshVisualPropertyChanged);
+			SubscribeToSignal(*SignalSubsystem, UE::Mass::Signals::MeshChanged);
+#if WITH_EDITOR
+			SubscribeToSignal(*SignalSubsystem, UE::Mass::Signals::SelectionChanged);
+			SubscribeToSignal(*SignalSubsystem, UE::Mass::Signals::LevelEditingStateChanged);
+#endif // WITH_EDITOR
+			SubscribeToSignal(*SignalSubsystem, UE::Mass::Signals::RenderStateDirty);
+		}
+	}
+}
+
+void UMassRenderStateDirtyUpdateProcessor::ConfigureQueries(const TSharedRef<FMassEntityManager>& EntityManager)
+{
+	EntityQuery.AddRequirement<FMassRenderStateFragment>(EMassFragmentAccess::ReadWrite);
+}
+
+void UMassRenderStateDirtyUpdateProcessor::SignalEntities(
+	FMassEntityManager& EntityManager,
+	FMassExecutionContext& ExecutionContext,
+	FMassSignalNameLookup& EntitySignals)
+{
+	const UObject* EntityManagerOwner = ExecutionContext.GetEntityManagerChecked().GetOwner();
+	check(EntityManagerOwner);
+
+	TArray<FName> SignalsForEntity;
+	EntityQuery.ForEachEntityChunk(ExecutionContext,
+		[&EntitySignals, &SignalsForEntity](FMassExecutionContext& Context)
+		{
+			const TArrayView<FMassRenderStateFragment>& RenderStateFragments = Context.GetMutableFragmentView<FMassRenderStateFragment>();
+
+			for (FMassExecutionContext::FEntityIterator EntityIt = Context.CreateEntityIterator(); EntityIt; ++EntityIt)
+			{
+				FMassPrimitiveRenderStateHelper& Helper = RenderStateFragments[*EntityIt].GetRenderStateHelper<FMassPrimitiveRenderStateHelper>();
+				SignalsForEntity.Reset();
+				EntitySignals.GetSignalsForEntity(Context.GetEntity(*EntityIt), SignalsForEntity);
+
+				bool bHandled = false;
+				if (SignalsForEntity.Num() == 1)
+				{
+					if (SignalsForEntity[0] == UE::Mass::Signals::TransformChanged)
+					{
+						Helper.UpdateTransform();
+						bHandled = true;
+					}
+#if WITH_EDITOR
+					else if (SignalsForEntity[0] == UE::Mass::Signals::SelectionChanged)
+					{
+						Helper.UpdateSelection();
+						bHandled = true;
+					}
+					else if (SignalsForEntity[0] == UE::Mass::Signals::LevelEditingStateChanged)
+					{
+						Helper.UpdateLevelInstanceEditingState();
+						bHandled = true;
+					}
+#endif // WITH_EDITOR
+				}
+
+				if (!bHandled)
+				{
+					Helper.DestroyRenderState();
+					Context.Defer().RemoveTag<FMassSceneProxyCreatedTag>(Context.GetEntity(*EntityIt));
+				}
+			}
+		});
+}
+
+// ----------------------------------------------------------------------//
+// UMassDestroyRenderStateProcessor
+//----------------------------------------------------------------------//
+UMassDestroyRenderStateProcessor::UMassDestroyRenderStateProcessor()
+	: EntityQuery(*this)
+{
+	ExecutionFlags = static_cast<int32>(EProcessorExecutionFlags::AllWorldModes);
+	ObservedType = FMassRenderStateFragment::StaticStruct();
+	ObservedOperations = EMassObservedOperationFlags::Remove;
+	bRequiresGameThreadExecution = true; // @todo revisit if the execution is required on game thread
+}
+
+void UMassDestroyRenderStateProcessor::ConfigureQueries(const TSharedRef<FMassEntityManager>& EntityManager)
+{
+	EntityQuery.AddRequirement<FMassRenderStateFragment>(EMassFragmentAccess::ReadWrite);
+}
+
+void UMassDestroyRenderStateProcessor::Execute(FMassEntityManager& EntityManager, FMassExecutionContext& ExecutionContext)
+{
+	FMassDestroyRenderStateContext BatchingContext(nullptr);
+
+	EntityQuery.ForEachEntityChunk(ExecutionContext,
+		[&BatchingContext](FMassExecutionContext& Context)
+		{
+			const TArrayView<FMassRenderStateFragment>& RenderStateFragments = Context.GetMutableFragmentView<FMassRenderStateFragment>();
+
+			for (FMassExecutionContext::FEntityIterator EntityIt = Context.CreateEntityIterator(); EntityIt; ++EntityIt)
+			{
+				FMassPrimitiveRenderStateHelper& Helper = RenderStateFragments[*EntityIt].GetRenderStateHelper<FMassPrimitiveRenderStateHelper>();
+				if (!BatchingContext.IsInitialized())
+				{
+					BatchingContext.Initialize(Helper.GetScene());
+				}
+				Helper.DestroyRenderState(&BatchingContext);
+				RenderStateFragments[*EntityIt].DestroyRenderStateHelper();
+			}
+		});
+
+	if (BatchingContext.IsInitialized())
+	{
+		BatchingContext.Process();
+	}
+}
diff --git a/Engine/Source/Runtime/Mass/MassEngine/Private/Mesh/MassEngineRenderStaticMeshProcessors.cpp b/Engine/Source/Runtime/Mass/MassEngine/Private/Mesh/MassEngineRenderStaticMeshProcessors.cpp
new file mode 100644
--- /dev/null
+++ b/Engine/Source/Runtime/Mass/MassEngine/Private/Mesh/MassEngineRenderStaticMeshProcessors.cpp
@@ -0,0 +1,196 @@
+// Copyright Epic Games, Inc. All Rights Reserved.
+
+#include "Mesh/MassEngineRenderStaticMeshProcessors.h"
+
+#include "AI/MassEngineNavigationProcessors.h"
+#include "MassDestroyRenderStateContext.h"
+#include "MassEngineTypes.h"
+#include "MassEntityFragments.h"
+#include "MassExecutionContext.h"
+#include "MassISMRenderStateHelper.h"
+#include "MassRenderStateHelper.h"
+#include "MassSignalSubsystem.h"
+#include "MassStaticMeshRenderStateHelper.h"
+#include "Mesh/MassEngineMeshFragments.h"
+#include "StaticMeshComponentHelper.h"
+#include "StaticMeshSceneProxyDesc.h"
+#include "Mesh/MassEngineMeshUtils.h"
+
+#include UE_INLINE_GENERATED_CPP_BY_NAME(MassEngineRenderStaticMeshProcessors)
+
+//----------------------------------------------------------------------//
+// UMassBaseStaticMeshSetupRenderStateProcessor
+//----------------------------------------------------------------------//
+UMassBaseStaticMeshSetupRenderStateProcessor::UMassBaseStaticMeshSetupRenderStateProcessor(const FObjectInitializer& ObjectInitializer)
+	: Super(ObjectInitializer)
+	, EntityQuery(*this)
+{
+	bRequiresGameThreadExecution = true; // @todo revisit if the execution is required on game thread
+	ExecutionFlags = static_cast<int32>(EProcessorExecutionFlags::AllWorldModes);
+	ExecutionOrder.ExecuteInGroup = UE::Mass::ProcessorGroupNames::SetupRenderState;
+}
+
+void UMassBaseStaticMeshSetupRenderStateProcessor::ConfigureQueries(const TSharedRef<FMassEntityManager>& EntityManager)
+{
+	EntityQuery.AddRequirement<FTransformFragment>(EMassFragmentAccess::ReadOnly);
+	EntityQuery.AddConstSharedRequirement<FMassStaticMeshFragment>();
+	EntityQuery.AddConstSharedRequirement<FMassVisualizationMeshFragment>();
+	EntityQuery.AddConstSharedRequirement<FMassOverrideMaterialsFragment>(EMassFragmentPresence::Optional);
+#if WITH_EDITOR
+	EntityQuery.AddConstSharedRequirement<FMassEditorVisualizationMeshFragment>(EMassFragmentPresence::Optional);
+#endif // WITH_EDITOR
+
+	EntityQuery.AddRequirement<FMassRenderPrimitiveFragment>(EMassFragmentAccess::ReadOnly, EMassFragmentPresence::None);
+	EntityQuery.AddRequirement<FMassRenderStateFragment>(EMassFragmentAccess::ReadOnly, EMassFragmentPresence::None);
+}
+
+//----------------------------------------------------------------------//
+// UMassStaticMeshSetupRenderStateProcessor
+//----------------------------------------------------------------------//
+void UMassStaticMeshSetupRenderStateProcessor::ConfigureQueries(const TSharedRef<FMassEntityManager>& EntityManager)
+{
+	Super::ConfigureQueries(EntityManager);
+	EntityQuery.AddTagRequirement<FMassRenderISMCandidateTag>(EMassFragmentPresence::None);
+	EntityQuery.AddRequirement<FMassRenderStaticMeshFragment>(EMassFragmentAccess::ReadOnly, EMassFragmentPresence::None);
+}
+
+void UMassStaticMeshSetupRenderStateProcessor::Execute(FMassEntityManager& EntityManager, FMassExecutionContext& ExecutionContext)
+{
+	EntityQuery.ForEachEntityChunk(ExecutionContext,
+		[&EntityManager](FMassExecutionContext& Context)
+		{
+			const FMassStaticMeshFragment& StaticMeshFragment = Context.GetConstSharedFragment<FMassStaticMeshFragment>();
+			const FMassVisualizationMeshFragment& MeshFragment = Context.GetConstSharedFragment<FMassVisualizationMeshFragment>();
+			const FMassOverrideMaterialsFragment* OverrideMaterialsFragment = Context.GetConstSharedFragmentPtr<FMassOverrideMaterialsFragment>();
+#if WITH_EDITOR
+			const FMassEditorVisualizationMeshFragment* EditorMeshFragment = Context.GetConstSharedFragmentPtr<FMassEditorVisualizationMeshFragment>();
+#endif // WITH_EDITOR
+			const TConstArrayView<FTransformFragment> TransformFragments = Context.GetFragmentView<FTransformFragment>();
+
+			for (FMassExecutionContext::FEntityIterator EntityIt = Context.CreateEntityIterator(); EntityIt; ++EntityIt)
+			{
+				FMassRenderPrimitiveFragment RenderPrimitiveFragment;
+				RenderPrimitiveFragment.bIsVisible = true;
+				RenderPrimitiveFragment.LocalBounds = StaticMeshFragment.Mesh.Get()->GetBounds();
+				RenderPrimitiveFragment.WorldBounds = RenderPrimitiveFragment.LocalBounds.TransformBy(TransformFragments[*EntityIt].GetTransform());
+
+				FMassRenderStaticMeshFragment RenderStaticMeshFragment;
+				checkf(RenderStaticMeshFragment.StaticMeshSceneProxyDesc, TEXT("Expecting a valid prxy desc here"));
+
+				UE::MassEngine::Mesh::InitializeStaticMeshSceneProxyDescFromFragment(StaticMeshFragment, MeshFragment, TransformFragments[*EntityIt], *RenderStaticMeshFragment.StaticMeshSceneProxyDesc);
+#if WITH_EDITOR
+				if (EditorMeshFragment)
+				{
+					UE::MassEngine::Mesh::InitializePrimitiveSceneProxyDescFromEditorFragment(*EditorMeshFragment, *RenderStaticMeshFragment.StaticMeshSceneProxyDesc);
+				}
+#endif // WITH_EDITOR
+
+				FMassRenderStateFragment RenderStateFragment;
+				RenderStateFragment.CreateRenderStateHelper<FMassStaticMeshRenderStateHelper>(Context.GetEntity(*EntityIt), &EntityManager, RenderPrimitiveFragment, OverrideMaterialsFragment, RenderStaticMeshFragment);
+
+				// Add new fragments
+				Context.Defer().PushCommand<FMassCommandAddFragmentInstances>(Context.GetEntity(*EntityIt), MoveTemp(RenderPrimitiveFragment), MoveTemp(RenderStaticMeshFragment), MoveTemp(RenderStateFragment));
+			}
+		});
+}
+
+//----------------------------------------------------------------------//
+// UMassStaticMeshUpdateRenderStateProcessor
+//----------------------------------------------------------------------//
+UMassStaticMeshUpdateRenderStateProcessor::UMassStaticMeshUpdateRenderStateProcessor(const FObjectInitializer& ObjectInitializer)
+	: Super(ObjectInitializer)
+{
+	bRequiresGameThreadExecution = true; // @todo revisit if the execution is required on game thread
+	ExecutionFlags = static_cast<int32>(EProcessorExecutionFlags::AllWorldModes);
+	ExecutionOrder.ExecuteInGroup = UE::Mass::ProcessorGroupNames::SetupRenderState;
+}
+
+void UMassStaticMeshUpdateRenderStateProcessor::InitializeInternal(UObject& Owner, const TSharedRef<FMassEntityManager>& EntityManager)
+{
+	Super::InitializeInternal(Owner, EntityManager);
+
+	if (const UWorld* World = GetWorld())
+	{
+		if (UMassSignalSubsystem* SignalSubsystem = World->GetSubsystem<UMassSignalSubsystem>())
+		{
+			SubscribeToSignal(*SignalSubsystem, UE::Mass::Signals::TransformChanged);
+			SubscribeToSignal(*SignalSubsystem, UE::Mass::Signals::MeshVisualPropertyChanged);
+			SubscribeToSignal(*SignalSubsystem, UE::Mass::Signals::MeshChanged);
+#if WITH_EDITOR
+			SubscribeToSignal(*SignalSubsystem, UE::Mass::Signals::SelectionChanged);
+			SubscribeToSignal(*SignalSubsystem, UE::Mass::Signals::LevelEditingStateChanged);
+#endif // WITH_EDITOR
+		}
+	}
+}
+
+void UMassStaticMeshUpdateRenderStateProcessor::ConfigureQueries(const TSharedRef<FMassEntityManager>& EntityManager)
+{
+	EntityQuery.AddConstSharedRequirement<FMassStaticMeshFragment>();
+	EntityQuery.AddConstSharedRequirement<FMassVisualizationMeshFragment>();
+#if WITH_EDITOR
+	EntityQuery.AddConstSharedRequirement<FMassEditorVisualizationMeshFragment>(EMassFragmentPresence::Optional);
+#endif // WITH_EDITOR
+	EntityQuery.AddRequirement<FTransformFragment>(EMassFragmentAccess::ReadOnly);
+
+	EntityQuery.AddRequirement<FMassRenderPrimitiveFragment>(EMassFragmentAccess::ReadWrite);
+	EntityQuery.AddRequirement<FMassRenderStaticMeshFragment>(EMassFragmentAccess::ReadWrite);
+}
+
+void UMassStaticMeshUpdateRenderStateProcessor::SignalEntities(
+	FMassEntityManager& EntityManager,
+	FMassExecutionContext& ExecutionContext,
+	FMassSignalNameLookup& EntitySignals)
+{
+	const UObject* EntityManagerOwner = ExecutionContext.GetEntityManagerChecked().GetOwner();
+	check(EntityManagerOwner);
+
+	TArray<FName> SignalsForEntity;
+	EntityQuery.ForEachEntityChunk(ExecutionContext,
+		[&EntitySignals, &SignalsForEntity](FMassExecutionContext& Context)
+		{
+			const FMassStaticMeshFragment& StaticMeshFragment = Context.GetConstSharedFragment<FMassStaticMeshFragment>();
+			const FMassVisualizationMeshFragment& MeshFragment = Context.GetConstSharedFragment<FMassVisualizationMeshFragment>();
+#if WITH_EDITOR
+			const FMassEditorVisualizationMeshFragment* EditorMeshFragment = Context.GetConstSharedFragmentPtr<FMassEditorVisualizationMeshFragment>();
+#endif // WITH_EDITOR
+			const TConstArrayView<FTransformFragment> TransformFragments = Context.GetFragmentView<FTransformFragment>();
+
+			const TArrayView<FMassRenderPrimitiveFragment>& RenderPrimitiveFragments = Context.GetMutableFragmentView<FMassRenderPrimitiveFragment>();
+			const TArrayView<FMassRenderStaticMeshFragment>& RenderStaticMeshFragments = Context.GetMutableFragmentView<FMassRenderStaticMeshFragment>();
+
+			for (FMassExecutionContext::FEntityIterator EntityIt = Context.CreateEntityIterator(); EntityIt; ++EntityIt)
+			{
+				SignalsForEntity.Reset();
+				EntitySignals.GetSignalsForEntity(Context.GetEntity(*EntityIt), SignalsForEntity);
+
+				const bool bTransformChanged = SignalsForEntity.Contains(UE::Mass::Signals::TransformChanged);
+				const bool bVisualPropertyChanged = SignalsForEntity.Contains(UE::Mass::Signals::MeshVisualPropertyChanged);
+				const bool bMeshChanged = SignalsForEntity.Contains(UE::Mass::Signals::MeshChanged);
+#if WITH_EDITOR
+				const bool bSelectionChanged = SignalsForEntity.Contains(UE::Mass::Signals::SelectionChanged);
+				const bool bLevelEditingStateChanged = SignalsForEntity.Contains(UE::Mass::Signals::LevelEditingStateChanged);
+#endif // WITH_EDITOR
+
+				if (bTransformChanged)
+				{
+					RenderPrimitiveFragments[*EntityIt].bIsVisible = true;
+					RenderPrimitiveFragments[*EntityIt].LocalBounds = StaticMeshFragment.Mesh.Get()->GetBounds();
+					RenderPrimitiveFragments[*EntityIt].WorldBounds = RenderPrimitiveFragments[*EntityIt].LocalBounds.TransformBy(TransformFragments[*EntityIt].GetTransform());
+				}
+
+				checkf(RenderStaticMeshFragments[*EntityIt].StaticMeshSceneProxyDesc, TEXT("Expecting a valid prxy desc here"));
+				if (bTransformChanged || bMeshChanged || bVisualPropertyChanged)
+				{
+					UE::MassEngine::Mesh::InitializeStaticMeshSceneProxyDescFromFragment(StaticMeshFragment, MeshFragment, TransformFragments[*EntityIt], *RenderStaticMeshFragments[*EntityIt].StaticMeshSceneProxyDesc);
+				}
+#if WITH_EDITOR
+				if (EditorMeshFragment && (bVisualPropertyChanged || bSelectionChanged || bLevelEditingStateChanged) )
+				{
+					UE::MassEngine::Mesh::InitializePrimitiveSceneProxyDescFromEditorFragment(*EditorMeshFragment, *RenderStaticMeshFragments[*EntityIt].StaticMeshSceneProxyDesc);
+				}
+#endif // WITH_EDITOR
+			}
+
+		});
+}
\ No newline at end of file
diff --git a/Engine/Source/Runtime/Mass/MassEngine/Private/Mesh/MassISMRenderStateHelper.cpp b/Engine/Source/Runtime/Mass/MassEngine/Private/Mesh/MassISMRenderStateHelper.cpp
new file mode 100644
--- /dev/null
+++ b/Engine/Source/Runtime/Mass/MassEngine/Private/Mesh/MassISMRenderStateHelper.cpp
@@ -0,0 +1,189 @@
+// Copyright Epic Games, Inc. All Rights Reserved.
+
+#include "MassISMRenderStateHelper.h"
+
+#include "InstancedStaticMeshComponentHelper.h"
+#include "MassEntityManager.h"
+#include "Mesh/MassEngineMeshFragments.h"
+#include "NaniteSceneProxy.h"
+#include "SceneInterface.h"
+
+FMassISMRenderStateHelper::FMassISMRenderStateHelper(FMassEntityHandle InEntityHandle, TNotNull<FMassEntityManager*> InEntityManager, const FMassRenderPrimitiveFragment& RenderPrimitiveFragment, const FMassOverrideMaterialsFragment* OverrideMaterialsFragment, const FMassRenderISMFragment& RenderStaticMeshFragment)
+	: Super(InEntityHandle, InEntityManager, RenderPrimitiveFragment, OverrideMaterialsFragment)
+{
+}
+
+void FMassISMRenderStateHelper::ResetSceneProxyDescUnsupportedProperties()
+{
+	Super::ResetSceneProxyDescUnsupportedProperties();
+
+	FMassRenderISMFragment& RenderISMFragment = GetMutableRenderISMFragment();
+
+	checkf(RenderISMFragment.InstancedStaticMeshSceneProxyDesc, TEXT("Expecting a valid scene proxy desc"));
+	RenderISMFragment.InstancedStaticMeshSceneProxyDesc->InstanceDataSceneProxy = nullptr;
+#if WITH_EDITOR
+	RenderISMFragment.InstancedStaticMeshSceneProxyDesc->bHasSelectedInstances = false;
+#endif// WITH_EDITOR 
+}
+
+TSharedPtr<FInstanceDataSceneProxy, ESPMode::ThreadSafe>& FMassISMRenderStateHelper::BuildInstanceData()
+{
+	FInstanceSceneDataBuffers InstanceSceneDataBuffers{};
+	FInstanceSceneDataBuffers::FAccessTag AccessTag(PointerHash(this));
+	auto View = InstanceSceneDataBuffers.BeginWriteAccess(AccessTag);
+	
+	// PrimitiveLocalToWorld
+	InstanceSceneDataBuffers.SetPrimitiveLocalToWorld(GetRenderMatrix(), AccessTag);
+	
+	// InstanceLocalBounds
+	const FPrimitiveMaterialPropertyDescriptor PrimitiveMaterialDesc = GetUsedMaterialPropertyDesc(GetScene()->GetShaderPlatform());
+	const float LocalAbsMaxDisplacement = FMath::Max(-PrimitiveMaterialDesc.MinMaxMaterialDisplacement.X, PrimitiveMaterialDesc.MinMaxMaterialDisplacement.Y) + PrimitiveMaterialDesc.MaxWorldPositionOffsetDisplacement;
+	const FVector3f PadExtent = FISMCInstanceDataSceneProxy::GetLocalBoundsPadExtent(View.PrimitiveToRelativeWorld, LocalAbsMaxDisplacement);
+	FRenderBounds InstanceLocalBounds = GetStaticMesh()->GetBounds();
+	InstanceLocalBounds.Min -= PadExtent;
+	InstanceLocalBounds.Max += PadExtent;
+	check(!View.Flags.bHasPerInstanceLocalBounds);
+	View.InstanceLocalBounds.Add(InstanceLocalBounds);
+
+	const FMassRenderISMFragment& RenderISMFragment = GetRenderISMFragment();
+
+	// LocalToPrimitiveRelativeWorld
+	View.InstanceToPrimitiveRelative.Reserve(RenderISMFragment.PerInstanceSMData.Num());
+	for (const auto& SM : RenderISMFragment.PerInstanceSMData)
+	{
+		FRenderTransform InstanceToPrimitive = SM.Transform;
+		FRenderTransform LocalToPrimitiveRelativeWorld = InstanceToPrimitive * View.PrimitiveToRelativeWorld;
+		LocalToPrimitiveRelativeWorld.Orthogonalize();
+		View.InstanceToPrimitiveRelative.Add(LocalToPrimitiveRelativeWorld);
+	}
+
+	// InstanceCustomData
+	View.InstanceCustomData = RenderISMFragment.PerInstanceSMCustomData;
+	View.NumCustomDataFloats = RenderISMFragment.PerInstanceSMCustomData.Num() && RenderISMFragment.PerInstanceSMData.Num() ? RenderISMFragment.PerInstanceSMCustomData.Num() / RenderISMFragment.PerInstanceSMData.Num() : 0;
+	View.Flags.bHasPerInstanceCustomData = PrimitiveMaterialDesc.bAnyMaterialHasPerInstanceCustomData && View.NumCustomDataFloats != 0;
+	if (!View.Flags.bHasPerInstanceCustomData)
+	{
+		View.NumCustomDataFloats = 0;
+		View.InstanceCustomData.Reset();
+	}
+
+	// InstanceRandomIDs
+	View.Flags.bHasPerInstanceRandom = PrimitiveMaterialDesc.bAnyMaterialHasPerInstanceRandom && (RenderISMFragment.PerInstanceSMData.Num() > 0);
+	if (View.Flags.bHasPerInstanceRandom)
+	{
+		InstanceRandomIDs.Reset(RenderISMFragment.PerInstanceSMData.Num());
+		InstanceRandomIDs.SetNum(RenderISMFragment.PerInstanceSMData.Num());
+		FRandomStream RandomStream = FRandomStream(RenderISMFragment.InstancingRandomSeed);
+		auto AdditionalRandomSeedsIt = RenderISMFragment.AdditionalRandomSeeds.CreateConstIterator();
+		int32 SeedResetIndex = AdditionalRandomSeedsIt ? AdditionalRandomSeedsIt->StartInstanceIndex : INDEX_NONE;
+		for (int32 Index = 0; Index < InstanceRandomIDs.Num(); ++Index)
+		{
+			// Reset the random stream if necessary
+			if (Index == SeedResetIndex)
+			{
+				RandomStream = FRandomStream(AdditionalRandomSeedsIt->RandomSeed);
+				++AdditionalRandomSeedsIt;
+				SeedResetIndex = AdditionalRandomSeedsIt ? AdditionalRandomSeedsIt->StartInstanceIndex : INDEX_NONE;
+			}
+			InstanceRandomIDs[Index] = RandomStream.GetFraction();
+		}
+	}
+	View.InstanceRandomIDs = InstanceRandomIDs;
+
+#if WITH_EDITOR
+	// Hit proxies
+	if (RenderISMFragment.PerInstanceHitProxy.Num())
+	{
+		UE_LOG(LogTemp, Warning, TEXT("CreateProxies"));
+		const FMassEditorVisualizationMeshFragment* EditorMeshFragment = GetEditorMeshFragment();
+		const bool bSelected = EditorMeshFragment ? EditorMeshFragment->bShouldRenderSelected : false;
+
+		View.Flags.bHasPerInstanceEditorData = true;
+		View.InstanceEditorData.Reserve(RenderISMFragment.PerInstanceHitProxy.Num());
+		for (HHitProxy* HitProxy : RenderISMFragment.PerInstanceHitProxy)
+		{
+			View.InstanceEditorData.Add(FInstanceEditorData::Pack(HitProxy->Id.GetColor(), bSelected));
+		}
+	}
+#endif // WITH_EDITOR
+	InstanceSceneDataBuffers.EndWriteAccess(AccessTag);
+	InstanceSceneDataBuffers.ValidateData();
+
+	DataProxy = MakeShared<FInstanceDataSceneProxy, ESPMode::ThreadSafe>(MoveTemp(InstanceSceneDataBuffers));
+	return DataProxy;
+}
+
+FPrimitiveSceneProxy* FMassISMRenderStateHelper::CreateStaticMeshSceneProxy(const Nanite::FMaterialAudit& NaniteMaterials, bool bCreateNanite)
+{
+	FMassRenderISMFragment& RenderISMFragment = GetMutableRenderISMFragment();
+
+	check(GetWorld());
+	check(RenderISMFragment.InstancedStaticMeshSceneProxyDesc);
+	check(RenderISMFragment.InstancedStaticMeshSceneProxyDesc->Scene);
+	check(RenderISMFragment.PerInstanceSMData.Num() > 0);
+
+	RenderISMFragment.InstancedStaticMeshSceneProxyDesc->InstanceDataSceneProxy = BuildInstanceData();
+	if (bCreateNanite)
+	{
+		PrimitiveSceneData.SceneProxy = ::new Nanite::FSceneProxy(NaniteMaterials, *RenderISMFragment.InstancedStaticMeshSceneProxyDesc);
+	}
+	else
+	{
+		PrimitiveSceneData.SceneProxy = ::new FInstancedStaticMeshSceneProxy(*RenderISMFragment.InstancedStaticMeshSceneProxyDesc, RenderISMFragment.InstancedStaticMeshSceneProxyDesc->FeatureLevel);
+	}
+	return PrimitiveSceneData.SceneProxy;
+}
+
+#if !WITH_EDITOR && UE_WITH_PSO_PRECACHING
+void FMassISMRenderStateHelper::CollectPSOPrecacheData(const FPSOPrecacheParams& BasePrecachePSOParams, FMaterialInterfacePSOPrecacheParamsList& OutParams)
+{
+	return FInstancedStaticMeshComponentHelper::CollectPSOPrecacheData(*this, BasePrecachePSOParams, OutParams);
+}
+#endif
+
+
+const FMassRenderISMFragment& FMassISMRenderStateHelper::GetRenderISMFragment() const
+{
+	return GetEntityManager().GetFragmentDataChecked<FMassRenderISMFragment>(EntityHandle);
+}
+
+FMassRenderISMFragment& FMassISMRenderStateHelper::GetMutableRenderISMFragment()
+{
+	return GetEntityManager().GetFragmentDataChecked<FMassRenderISMFragment>(EntityHandle);
+}
+
+void FMassISMRenderStateHelper::InitializeSceneProxyDescDynamicProperties()
+{
+	FMassBaseStaticMeshRenderStateHelper::InitializeSceneProxyDescDynamicProperties();
+
+#if WITH_EDITOR
+	FStaticMeshSceneProxyDesc& SceneProxyDesc = GetStaticMeshSceneProxyDesc();
+
+	FMassRenderISMFragment& RenderISMFragment = GetMutableRenderISMFragment();
+	SceneProxyDesc.bHasPerInstanceHitProxies = RenderISMFragment.PerInstanceHitProxy.Num() > 0;
+#endif // WITH_EDITOR
+}
+
+FPrimitiveSceneProxyDesc& FMassISMRenderStateHelper::GetSceneProxyDesc()
+{
+	return GetStaticMeshSceneProxyDesc();
+}
+
+const FPrimitiveSceneProxyDesc& FMassISMRenderStateHelper::GetSceneProxyDesc() const
+{
+	return GetStaticMeshSceneProxyDesc();
+}
+
+FStaticMeshSceneProxyDesc& FMassISMRenderStateHelper::GetStaticMeshSceneProxyDesc()
+{
+	FMassRenderISMFragment& RenderISMFragment = GetMutableRenderISMFragment();
+	checkf(RenderISMFragment.InstancedStaticMeshSceneProxyDesc, TEXT("Expecting a valid scene proxy desc"));
+	return *RenderISMFragment.InstancedStaticMeshSceneProxyDesc;
+}
+
+const FStaticMeshSceneProxyDesc& FMassISMRenderStateHelper::GetStaticMeshSceneProxyDesc() const
+{
+	const FMassRenderISMFragment& RenderISMFragment = GetRenderISMFragment();
+	checkf(RenderISMFragment.InstancedStaticMeshSceneProxyDesc, TEXT("Expecting a valid scene proxy desc"));
+	return *RenderISMFragment.InstancedStaticMeshSceneProxyDesc;
+}
diff --git a/Engine/Source/Runtime/Mass/MassEngine/Private/Mesh/MassISMRenderStateHelper.h b/Engine/Source/Runtime/Mass/MassEngine/Private/Mesh/MassISMRenderStateHelper.h
new file mode 100644
--- /dev/null
+++ b/Engine/Source/Runtime/Mass/MassEngine/Private/Mesh/MassISMRenderStateHelper.h
@@ -0,0 +1,50 @@
+// Copyright Epic Games, Inc. All Rights Reserved.
+
+#pragma once
+
+#include "MassStaticMeshRenderStateHelper.h"
+#include "InstancedStaticMeshSceneProxyDesc.h"
+
+struct FMassRenderISMFragment;
+
+/**
+ * Helper to handle all the communication to the renderer from Mass for an instanced static mesh
+ */
+struct FMassISMRenderStateHelper : public FMassBaseStaticMeshRenderStateHelper
+{
+public:
+	using Super = FMassBaseStaticMeshRenderStateHelper ;
+
+	FMassISMRenderStateHelper(FMassEntityHandle InEntityHandle, TNotNull<FMassEntityManager*> InEntityManager, const FMassRenderPrimitiveFragment& RenderPrimitiveFragment, const FMassOverrideMaterialsFragment* OverrideMaterialsFragment, const FMassRenderISMFragment& RenderStaticMeshFragment);
+	virtual ~FMassISMRenderStateHelper() override = default;
+
+protected:
+	//~ Begin FMassPrimitiveRenderStateHelper interface
+	virtual FPrimitiveSceneProxyDesc& GetSceneProxyDesc() override;
+	virtual const FPrimitiveSceneProxyDesc& GetSceneProxyDesc() const override;
+	virtual void ResetSceneProxyDescUnsupportedProperties() override;
+	//~ End FMassPrimitiveRenderStateHelper interface
+
+	//~ Begin FMassBaseStaticMeshRenderStateHelper interface
+	virtual FStaticMeshSceneProxyDesc& GetStaticMeshSceneProxyDesc() override;
+	virtual const FStaticMeshSceneProxyDesc& GetStaticMeshSceneProxyDesc() const override;
+	virtual FPrimitiveSceneProxy* CreateStaticMeshSceneProxy(const Nanite::FMaterialAudit& NaniteMaterials, bool bCreateNanite) override;
+#if !WITH_EDITOR && UE_WITH_PSO_PRECACHING
+	virtual void CollectPSOPrecacheData(const FPSOPrecacheParams& BasePrecachePSOParams, FMaterialInterfacePSOPrecacheParamsList& OutParams) override;
+#endif
+	//~ End FMassBaseStaticMeshRenderStateHelper interface
+
+	TSharedPtr<FInstanceDataSceneProxy, ESPMode::ThreadSafe>& BuildInstanceData();
+
+	const FMassRenderISMFragment& GetRenderISMFragment() const;
+	FMassRenderISMFragment& GetMutableRenderISMFragment();
+
+protected:
+
+	virtual void InitializeSceneProxyDescDynamicProperties() override;
+
+	TSharedPtr<FInstanceDataSceneProxy, ESPMode::ThreadSafe> DataProxy{};
+	TArray<float> InstanceRandomIDs;
+
+	friend class FInstancedStaticMeshComponentHelper;
+};
diff --git a/Engine/Source/Runtime/Mass/MassEngine/Private/Mesh/MassMeshRenderStateHelper.cpp b/Engine/Source/Runtime/Mass/MassEngine/Private/Mesh/MassMeshRenderStateHelper.cpp
new file mode 100644
--- /dev/null
+++ b/Engine/Source/Runtime/Mass/MassEngine/Private/Mesh/MassMeshRenderStateHelper.cpp
@@ -0,0 +1,48 @@
+// Copyright Epic Games, Inc. All Rights Reserved.
+
+#include "MassMeshRenderStateHelper.h"
+
+#include "MassEntityManager.h"
+#include "MeshComponentHelper.h"
+#include "Mesh/MassEngineMeshFragments.h"
+
+FMassMeshRenderStateHelper::FMassMeshRenderStateHelper(FMassEntityHandle InEntityHandle, TNotNull<FMassEntityManager*> InEntityManager, const FMassRenderPrimitiveFragment& RenderPrimitiveFragment, const FMassOverrideMaterialsFragment* OverrideMaterialsFragment)
+	: Super(InEntityHandle, InEntityManager, RenderPrimitiveFragment)
+{
+}	
+
+void FMassMeshRenderStateHelper::GetMaterialSlotsOverlayMaterial(TArray<TObjectPtr<UMaterialInterface>>& OutMaterialSlotOverlayMaterials) const
+{
+	FMeshComponentHelper::GetMaterialSlotsOverlayMaterial(*this, OutMaterialSlotOverlayMaterials);
+}
+
+FMaterialRelevance FMassMeshRenderStateHelper::GetMaterialRelevance(EShaderPlatform InShaderPlatform) const
+{
+	return FMeshComponentHelper::GetMaterialRelevance(*this, InShaderPlatform);
+}
+
+const FMassOverrideMaterialsFragment* FMassMeshRenderStateHelper::GetOverrideMaterialsFragment() const
+{
+	return GetEntityManager().GetConstSharedFragmentDataPtr<FMassOverrideMaterialsFragment>(EntityHandle);
+}
+
+const FMassVisualizationMeshFragment& FMassMeshRenderStateHelper::GetMeshFragment() const
+{
+	return GetEntityManager().GetConstSharedFragmentDataChecked<FMassVisualizationMeshFragment>(EntityHandle);
+}
+
+#if WITH_EDITOR
+const FMassEditorVisualizationMeshFragment* FMassMeshRenderStateHelper::GetEditorMeshFragment() const
+{
+	return GetEntityManager().GetConstSharedFragmentDataPtr<FMassEditorVisualizationMeshFragment>(EntityHandle);
+}
+#endif // WITH_EDITOR
+
+TConstArrayView<TObjectPtr<UMaterialInterface>> FMassMeshRenderStateHelper::GetOverrideMaterials() const
+{
+	if (const FMassOverrideMaterialsFragment* OverrideMaterialFragment = GetOverrideMaterialsFragment())
+	{
+		return OverrideMaterialFragment->OverrideMaterials;
+	}
+	return TConstArrayView<TObjectPtr<UMaterialInterface>>();
+}
diff --git a/Engine/Source/Runtime/Mass/MassEngine/Private/Mesh/MassMeshRenderStateHelper.h b/Engine/Source/Runtime/Mass/MassEngine/Private/Mesh/MassMeshRenderStateHelper.h
new file mode 100644
--- /dev/null
+++ b/Engine/Source/Runtime/Mass/MassEngine/Private/Mesh/MassMeshRenderStateHelper.h
@@ -0,0 +1,41 @@
+// Copyright Epic Games, Inc. All Rights Reserved.
+
+#pragma once
+
+#include "MassPrimitiveRenderStateHelper.h"
+
+struct FMassEditorVisualizationMeshFragment;
+struct FMassVisualizationMeshFragment;
+class UMaterialInterface;
+struct FMaterialRelevance;
+struct FMassOverrideMaterialsFragment;
+
+/**
+ * Helper to handle all the communication to the renderer from Mass for all type of meshes 
+ */
+struct FMassMeshRenderStateHelper : public FMassPrimitiveRenderStateHelper
+{
+public:
+	using Super = FMassPrimitiveRenderStateHelper;
+
+	FMassMeshRenderStateHelper(FMassEntityHandle InEntityHandle, TNotNull<FMassEntityManager*> InEntityManager, const FMassRenderPrimitiveFragment& RenderPrimitiveFragment, const FMassOverrideMaterialsFragment* OverrideMaterialsFragment);
+	virtual ~FMassMeshRenderStateHelper() override = default;
+
+	FMaterialRelevance GetMaterialRelevance(EShaderPlatform InShaderPlatform) const;
+
+	void GetMaterialSlotsOverlayMaterial(TArray<TObjectPtr<UMaterialInterface>>& OutMaterialSlotOverlayMaterials) const;
+
+	//Get the asset material slot overlay material
+	virtual void GetDefaultMaterialSlotsOverlayMaterial(TArray<TObjectPtr<UMaterialInterface>>& AssetMaterialSlotOverlayMaterials) const = 0;
+
+	virtual const TArray<TObjectPtr<UMaterialInterface>>& GetComponentMaterialSlotsOverlayMaterial() const = 0;
+	virtual UMaterialInterface* GetOverlayMaterial() const = 0;
+
+	const FMassOverrideMaterialsFragment* GetOverrideMaterialsFragment() const;
+	const FMassVisualizationMeshFragment& GetMeshFragment() const;
+#if WITH_EDITOR
+	const FMassEditorVisualizationMeshFragment* GetEditorMeshFragment() const;
+#endif // WITH_EDITOR
+
+	TConstArrayView<TObjectPtr<UMaterialInterface>> GetOverrideMaterials() const;
+};
\ No newline at end of file
diff --git a/Engine/Source/Runtime/Mass/MassEngine/Private/Mesh/MassPrimitiveRenderStateHelper.cpp b/Engine/Source/Runtime/Mass/MassEngine/Private/Mesh/MassPrimitiveRenderStateHelper.cpp
new file mode 100644
--- /dev/null
+++ b/Engine/Source/Runtime/Mass/MassEngine/Private/Mesh/MassPrimitiveRenderStateHelper.cpp
@@ -0,0 +1,634 @@
+// Copyright Epic Games, Inc. All Rights Reserved.
+
+#include "MassPrimitiveRenderStateHelper.h"
+
+#include "MassDestroyRenderStateContext.h"
+#include "MassEntityFragments.h"
+#include "MassEntityManager.h"
+#include "Mesh/MassEngineMeshFragments.h"
+#include "PrimitiveComponentHelper.h"
+#include "PrimitiveSceneDesc.h"
+#include "PrimitiveSceneProxyDesc.h"
+#include "SceneInterface.h"
+#include "UnrealEngine.h"
+#include "Elements/Framework/TypedElementHandle.h"
+
+
+#if WITH_EDITOR
+#include "ObjectCacheEventSink.h"
+#endif //WITH_EDITOR
+
+//----------------------------------------------------------------------//
+// FMassPrimitiveRenderStateHelper
+//----------------------------------------------------------------------//
+FMassPrimitiveRenderStateHelper::FMassPrimitiveRenderStateHelper(FMassEntityHandle InEntityHandle, TNotNull<FMassEntityManager*> InEntityManager, const FMassRenderPrimitiveFragment& RenderPrimitiveFragment)
+	: Super(InEntityHandle, InEntityManager)
+{
+	// Copying data from fragment as scene proxy decs take a view on this array
+	CustomPrimitiveData = RenderPrimitiveFragment.CustomPrimitiveData;
+}
+
+const FMassRenderPrimitiveFragment& FMassPrimitiveRenderStateHelper::GetRenderPrimitiveFragment() const
+{
+	return GetEntityManager().GetFragmentDataChecked<FMassRenderPrimitiveFragment>(EntityHandle);
+}
+
+FMassRenderPrimitiveFragment& FMassPrimitiveRenderStateHelper::GetMutableRenderPrimitiveFragment()
+{
+	return GetEntityManager().GetFragmentDataChecked<FMassRenderPrimitiveFragment>(EntityHandle);
+}
+
+void FMassPrimitiveRenderStateHelper::ResetSceneProxyDescUnsupportedProperties()
+{
+	// Unsupported properties
+	FPrimitiveSceneProxyDesc& SceneProxyDesc = GetSceneProxyDesc();
+	//SceneProxyDesc.bLevelInstanceEditingState = false; //@Todo FastGeo will need to turn this off outside of this method
+	//SceneProxyDesc.bSelectable = false; //@Todo FastGeo will need to turn this off outside of this method
+	SceneProxyDesc.bUseEditorCompositing = false;
+	SceneProxyDesc.bIsBeingMovedByEditor = false;
+	//SceneProxyDesc.bSelected = false; //@Todo FastGeo will need to turn this off outside of this method
+	//SceneProxyDesc.bIndividuallySelected = false; //@Todo FastGeo will need to turn this off outside of this method
+	//SceneProxyDesc.bShouldRenderSelected = false; //@Todo FastGeo will need to turn this off outside of this method
+	SceneProxyDesc.bWantsEditorEffects = false;
+	SceneProxyDesc.bIsHiddenEd = false;
+	SceneProxyDesc.bIsOwnerEditorOnly = false;
+#if WITH_EDITOR
+	SceneProxyDesc.bIsOwnedByFoliage = false;
+#endif // WITH_EDITOR
+	SceneProxyDesc.HiddenEditorViews = 0;
+#if WITH_EDITOR
+	SceneProxyDesc.OverlayColor = FColor(EForceInit::ForceInitToZero);
+#endif // WITH_EDITOR
+	SceneProxyDesc.Component = nullptr;
+
+	// Properties that will be initialized by InitializeSceneProxyDescDynamicProperties
+	SceneProxyDesc.ComponentId = FPrimitiveComponentId();
+	SceneProxyDesc.StatId = TStatId();
+	SceneProxyDesc.Owner = nullptr;
+	SceneProxyDesc.World = nullptr;
+	SceneProxyDesc.CustomPrimitiveData = nullptr;
+	SceneProxyDesc.Scene = nullptr;
+	SceneProxyDesc.PrimitiveComponentInterface = nullptr;
+	SceneProxyDesc.FeatureLevel = ERHIFeatureLevel::Num;
+	SceneProxyDesc.RuntimeVirtualTextures = TArrayView<URuntimeVirtualTexture*>();
+	SceneProxyDesc.bIsVisible = false;
+	SceneProxyDesc.bUsePSOPrecacheFallbackMaterial = false;
+#if MESH_DRAW_COMMAND_STATS
+	SceneProxyDesc.MeshDrawCommandStatsCategory = NAME_None;
+#endif // MESH_DRAW_COMMAND_STATS
+}
+
+bool FMassPrimitiveRenderStateHelper::IsDrawnInGame() const
+{
+	// Drawn in game must consider both the component bIsVisible flag AND the bIsHidden flag (which actually originates from the actor bHiddenInGame property)
+	// This logic mimics what is done to initialize FPrimitiveSceneProxy::DrawInGame
+	return GetSceneProxyDesc().bIsVisible && !GetSceneProxyDesc().bIsHidden;
+}
+
+EComponentMobility::Type FMassPrimitiveRenderStateHelper::GetMobility() const
+{
+	return GetSceneProxyDesc().Mobility;
+}
+
+void FMassPrimitiveRenderStateHelper::UpdateVisibility()
+{
+	const FMassRenderPrimitiveFragment& RenderPrimitiveFragment = GetRenderPrimitiveFragment();
+
+	// Update SceneProxyDesc.bIsVisible as it's dependant on component and component cluster visibility
+	FPrimitiveSceneProxyDesc& SceneProxyDesc = GetSceneProxyDesc();
+	SceneProxyDesc.bIsVisible = RenderPrimitiveFragment.bIsVisible;
+}
+
+void FMassPrimitiveRenderStateHelper::UpdateTransform()
+{
+	FMassRenderPrimitiveFragment& RenderPrimitiveFragment = GetMutableRenderPrimitiveFragment();
+
+	FSceneInterface* Scene = GetScene();
+	checkf(Scene, TEXT("Expecting a scene if this method is called"));
+
+	SceneDesc.RenderMatrix = GetRenderMatrix();
+	SceneDesc.AttachmentRootPosition = GetTransform().GetTranslation();
+	SceneDesc.LocalBounds = RenderPrimitiveFragment.LocalBounds;
+	SceneDesc.Bounds = GetBounds();
+
+	Scene->UpdatePrimitiveTransform(&SceneDesc);
+}
+
+#if WITH_EDITOR
+void FMassPrimitiveRenderStateHelper::UpdateSelection()
+{
+	FPrimitiveSceneProxyDesc& SceneProxyDesc = GetSceneProxyDesc();
+	if (PrimitiveSceneData.SceneProxy)
+	{
+		PrimitiveSceneData.SceneProxy->SetSelection_GameThread(SceneProxyDesc.ShouldRenderSelected());
+	}
+}
+
+void FMassPrimitiveRenderStateHelper::UpdateLevelInstanceEditingState()
+{
+	FPrimitiveSceneProxyDesc& SceneProxyDesc = GetSceneProxyDesc();
+	if (PrimitiveSceneData.SceneProxy)
+	{
+		PrimitiveSceneData.SceneProxy->SetLevelInstanceEditingState_GameThread(SceneProxyDesc.GetLevelInstanceEditingState());
+	}
+}
+#endif// WITH_EDITOR
+
+void FMassPrimitiveRenderStateHelper::InitializeSceneProxyDescDynamicProperties()
+{
+	check(GetWorld());
+	check(GetScene());
+
+	ResetSceneProxyDescUnsupportedProperties();
+
+	// Initialize non-serialized properties
+	FPrimitiveSceneProxyDesc& SceneProxyDesc = GetSceneProxyDesc();
+	SceneProxyDesc.ComponentId = GetPrimitiveSceneId();
+	UObject const* AdditionalStatObjectPtr = AdditionalStatObject();
+	SceneProxyDesc.StatId = AdditionalStatObjectPtr ? AdditionalStatObjectPtr->GetStatID(true) : TStatId();
+	SceneProxyDesc.Owner = GetOwner();
+#if !WITH_STATE_STREAM
+	SceneProxyDesc.World = GetWorld();
+#else
+	SceneProxyDesc.World = nullptr;
+#endif // !WITH_STATE_STREAM
+	SceneProxyDesc.CustomPrimitiveData = &CustomPrimitiveData;
+	SceneProxyDesc.Scene = GetScene();
+	SceneProxyDesc.PrimitiveComponentInterface = this;
+	SceneProxyDesc.FeatureLevel = SceneProxyDesc.Scene->GetFeatureLevel();
+	TConstArrayView<URuntimeVirtualTexture*> RuntimeVirtualTextures = GetRuntimeVirtualTextures();
+	SceneProxyDesc.RuntimeVirtualTextures = (TArrayView<URuntimeVirtualTexture*>&)RuntimeVirtualTextures; // ugly cast, but just making it work for now. It would mean to change the SceneProxyDesc to fix this.
+	SceneProxyDesc.bUsePSOPrecacheFallbackMaterial = UsePSOPrecacheFallbackMaterial();
+#if MESH_DRAW_COMMAND_STATS
+	static FName NAME_FastGeoPrimitiveComponent(TEXT("FastGeoPrimitiveComponent"));
+	SceneProxyDesc.MeshDrawCommandStatsCategory = NAME_FastGeoPrimitiveComponent;
+#endif // MESH_DRAW_COMMAND_STATS
+	//check(SceneProxyDesc.bIsInstancedStaticMesh == (this->IsA<FFastGeoInstancedStaticMeshComponent>() || this->IsA<FFastGeoProceduralISMComponent>()));
+	UpdateVisibility();
+}
+
+FPrimitiveSceneDesc FMassPrimitiveRenderStateHelper::BuildSceneDesc()
+{
+	FMassRenderPrimitiveFragment& RenderPrimitiveFragment = GetMutableRenderPrimitiveFragment();
+
+	check(GetSceneProxy());
+
+	FPrimitiveSceneProxyDesc& SceneProxyDesc = GetSceneProxyDesc();
+	SceneDesc.SceneProxy = GetSceneProxy();
+	SceneDesc.ProxyDesc = &SceneProxyDesc;
+	SceneDesc.PrimitiveSceneData = &PrimitiveSceneData;
+#if WITH_EDITOR
+	SceneDesc.PrimitiveComponentInterface = this;
+#endif
+	SceneDesc.RenderMatrix = GetRenderMatrix();
+	SceneDesc.AttachmentRootPosition = GetTransform().GetTranslation();
+	SceneDesc.LocalBounds = RenderPrimitiveFragment.LocalBounds;
+	SceneDesc.Bounds = GetBounds();
+	SceneDesc.Mobility = SceneProxyDesc.Mobility;
+
+	return SceneDesc;
+}
+
+void FMassPrimitiveRenderStateHelper::CreateRenderState(FRegisterComponentContext* Context)
+{
+	//FWriteScopeLock WriteLock(*Lock.Get());
+	check(ProxyState == EProxyCreationState::None);
+	ProxyState = EProxyCreationState::Creating;
+	bRenderStateDirty = false;
+
+#if WITH_EDITOR
+	NotifyRenderStateChanged();
+#endif
+	
+	FSceneInterface* Scene = GetScene();
+	check(Scene);
+
+	ESceneProxyCreationError Error = ESceneProxyCreationError::None;
+	if (FPrimitiveSceneProxy* SceneProxy = CreateSceneProxy(&Error))
+	{
+		//SceneProxy->SetPrimitiveColor(GetDebugColor());
+		check(GetSceneProxy());
+		FPrimitiveSceneDesc Desc = BuildSceneDesc();
+		Scene->AddPrimitive(&Desc);
+
+		ProxyState = EProxyCreationState::Created;
+	}
+	else if (Error == ESceneProxyCreationError::WaitingPSOs)
+	{
+		ProxyState = EProxyCreationState::Delayed;
+	}
+	else
+	{
+		ProxyState = EProxyCreationState::None;
+	}
+}
+
+FString FMassPrimitiveRenderStateHelper::GetName() const
+{
+	return GetNameSafe(GetEntityManager().GetOwner());
+}
+
+FString FMassPrimitiveRenderStateHelper::GetFullName() const
+{
+	return GetFullNameSafe(GetEntityManager().GetOwner());
+}
+
+void FMassPrimitiveRenderStateHelper::DestroyRenderState(FMassDestroyRenderStateContext* Context)
+{
+	//FWriteScopeLock WriteLock(*Lock.Get());
+	if (GetSceneProxy())
+	{
+		FSceneInterface* Scene = GetScene();
+		check(Scene);
+
+		check(ProxyState == EProxyCreationState::Created);
+
+		FMassDestroyRenderStateContext::DestroyProxy(Context, GetSceneProxy(), Scene);
+
+		PrimitiveSceneData.SceneProxy = nullptr;
+
+		ProxyState = EProxyCreationState::None;
+
+		bRenderStateDirty = false;
+
+#if WITH_EDITOR
+		NotifyRenderStateChanged();
+#endif
+	}
+	else
+	{
+		check(ProxyState == EProxyCreationState::None || ProxyState == EProxyCreationState::Delayed);
+		ProxyState = EProxyCreationState::None;
+	}
+}
+
+#if WITH_EDITOR
+void FMassPrimitiveRenderStateHelper::NotifyRenderStateChanged()
+{
+	FObjectCacheEventSink::NotifyRenderStateChanged_Concurrent(this);
+}
+#endif
+
+bool FMassPrimitiveRenderStateHelper::IsRenderStateDirty() const
+{
+	return Super::IsRenderStateDirty();
+}
+
+bool FMassPrimitiveRenderStateHelper::IsRenderStateCreated() const
+{
+	return Super::IsRenderStateCreated();
+}
+
+bool FMassPrimitiveRenderStateHelper::ShouldCreateRenderState() const
+{
+	const FMassRenderPrimitiveFragment& RenderPrimitiveFragment = GetRenderPrimitiveFragment();
+
+	if (!FApp::CanEverRender())
+	{
+		return false;
+	}
+
+	// If the detail mode setting allows it, add it to the scene.
+	const bool bDetailModeAllowsRendering = RenderPrimitiveFragment.DetailMode <= GetCachedScalabilityCVars().DetailMode;
+	if (!bDetailModeAllowsRendering)
+	{
+		return false;
+	}
+
+	const FPrimitiveSceneProxyDesc& SceneProxyDesc = GetSceneProxyDesc();
+	const bool bShouldCreateRenderState = (RenderPrimitiveFragment.bIsVisible && !SceneProxyDesc.bIsHidden) || SceneProxyDesc.bCastHiddenShadow || SceneProxyDesc.bAffectIndirectLightingWhileHidden || SceneProxyDesc.bRayTracingFarField;
+	return bShouldCreateRenderState;
+}
+
+bool FMassPrimitiveRenderStateHelper::IsRegistered() const
+{
+	// If the Mass entity is created it is because the element is registered in the world/scene
+	return true;
+}
+
+bool FMassPrimitiveRenderStateHelper::IsUnreachable() const
+{
+	const UObject* OwnerPtr = GetEntityManager().GetOwner();
+	return !OwnerPtr || OwnerPtr->IsUnreachable();
+}
+
+bool FMassPrimitiveRenderStateHelper::IsStaticMobility() const
+{
+	return true;
+}
+
+bool FMassPrimitiveRenderStateHelper::IsMipStreamingForced() const
+{
+	return false;
+}
+
+UWorld* FMassPrimitiveRenderStateHelper::GetWorld() const
+{
+	return Super::GetWorld();
+}
+
+FSceneInterface* FMassPrimitiveRenderStateHelper::GetScene() const
+{
+	return Super::GetScene();
+}
+
+TConstArrayView<URuntimeVirtualTexture*> FMassPrimitiveRenderStateHelper::GetRuntimeVirtualTextures() const 
+{
+	const FMassRenderPrimitiveFragment& RenderPrimitiveFragment = GetRenderPrimitiveFragment();
+	return RenderPrimitiveFragment.RuntimeVirtualTextures;
+}
+
+bool FMassPrimitiveRenderStateHelper::IsFirstPersonRelevant() const
+{
+	return GetSceneProxyDesc().IsFirstPersonRelevant();
+}
+
+FPrimitiveComponentId FMassPrimitiveRenderStateHelper::GetPrimitiveSceneId() const 
+{ 
+	return PrimitiveSceneData.PrimitiveSceneId;
+}
+
+FPrimitiveSceneProxy* FMassPrimitiveRenderStateHelper::GetSceneProxy() const
+{
+	return PrimitiveSceneData.SceneProxy;
+}
+
+void FMassPrimitiveRenderStateHelper::MarkRenderStateDirty()
+{
+	Super::MarkRenderStateDirty();
+}
+
+void FMassPrimitiveRenderStateHelper::DestroyRenderState()
+{
+	DestroyRenderState(nullptr);
+}
+
+FTransform FMassPrimitiveRenderStateHelper::GetTransform() const
+{
+	const FTransformFragment& TransformFragment = GetEntityManager().GetFragmentDataChecked<FTransformFragment>(EntityHandle);
+	return TransformFragment.GetTransform();
+}
+
+const FBoxSphereBounds& FMassPrimitiveRenderStateHelper::GetBounds() const
+{
+	const FMassRenderPrimitiveFragment& RenderPrimitiveFragment = GetRenderPrimitiveFragment();
+	return RenderPrimitiveFragment.WorldBounds;
+}
+
+void FMassPrimitiveRenderStateHelper::GetPrimitiveStats(FPrimitiveStats& PrimitiveStats) const
+{
+}
+
+UObject* FMassPrimitiveRenderStateHelper::GetUObject()
+{
+	return GetEntityManager().GetOwner();
+}
+
+const UObject* FMassPrimitiveRenderStateHelper::GetUObject() const
+{
+	return GetEntityManager().GetOwner();
+}
+
+FMatrix FMassPrimitiveRenderStateHelper::GetRenderMatrix() const
+{
+	return GetTransform().ToMatrixWithScale();
+}
+
+float FMassPrimitiveRenderStateHelper::GetLastRenderTimeOnScreen() const
+{
+	return PrimitiveSceneData.GetLastRenderTimeOnScreen();
+}
+
+void FMassPrimitiveRenderStateHelper::SetCollisionEnabled(bool bInCollisionEnabled)
+{
+	FPrimitiveSceneProxyDesc& SceneProxyDesc = GetSceneProxyDesc();
+	SceneProxyDesc.bCollisionEnabled = bInCollisionEnabled;
+}
+
+#if !WITH_EDITOR && UE_WITH_PSO_PRECACHING
+void FMassPrimitiveRenderStateHelper::OnPrecacheFinished(int32 JobSetThatJustCompleted)
+{
+	QUICK_SCOPE_CYCLE_COUNTER(STAT_PSOPrecacheFinishedTask);
+	int32 CurrJobSetCompleted = LatestPSOPrecacheJobSetCompleted.load();
+	while (CurrJobSetCompleted < JobSetThatJustCompleted && !LatestPSOPrecacheJobSetCompleted.compare_exchange_weak(CurrJobSetCompleted, JobSetThatJustCompleted)) {}
+	MarkRenderStateDirty();
+}
+
+void FMassPrimitiveRenderStateHelper::RequestRecreateRenderStateWhenPSOPrecacheFinished(const FGraphEventArray& PSOPrecacheCompileEvents)
+{
+	// If the proxy creation strategy relies on knowing when the precached PSO has been compiled,
+	// schedule a task to mark the render state dirty when all PSOs are compiled so the proxy gets recreated.
+	if (GetPSOPrecacheProxyCreationStrategy() != EPSOPrecacheProxyCreationStrategy::AlwaysCreate)
+	{
+		LatestPSOPrecacheJobSet++;
+		MarkRenderStateDirty();
+	}
+	bPSOPrecacheCalled = true;
+}
+
+void FMassPrimitiveRenderStateHelper::SetupPrecachePSOParams(FPSOPrecacheParams& Params)
+{
+	const FPrimitiveSceneProxyDesc& SceneProxyDesc = GetSceneProxyDesc();
+	auto IsPrecomputedLightingValid = []() { return false; };
+	Params.bRenderInMainPass = SceneProxyDesc.bRenderInMainPass;
+	Params.bRenderInDepthPass = SceneProxyDesc.bRenderInDepthPass;
+	Params.bStaticLighting = SceneProxyDesc.bHasStaticLighting;
+	Params.bUsesIndirectLightingCache = Params.bStaticLighting && SceneProxyDesc.IndirectLightingCacheQuality != ILCQ_Off && (!IsPrecomputedLightingValid() || SceneProxyDesc.LightmapType == ELightmapType::ForceVolumetric);
+	Params.bAffectDynamicIndirectLighting = SceneProxyDesc.bAffectDynamicIndirectLighting;
+	Params.bCastShadow = SceneProxyDesc.CastShadow;
+	// Custom depth can be toggled at runtime with PSO precache call so assume it might be needed when depth pass is needed
+	// Ideally precache those with lower priority and don't wait on these (UE-174426)
+	Params.bRenderCustomDepth = SceneProxyDesc.bRenderInDepthPass;
+	Params.bCastShadowAsTwoSided = SceneProxyDesc.bCastShadowAsTwoSided;
+	Params.SetMobility(SceneProxyDesc.Mobility);
+	Params.SetStencilWriteMask(FRendererStencilMaskEvaluation::ToStencilMask(SceneProxyDesc.CustomDepthStencilWriteMask));
+
+	TArray<UMaterialInterface*> UsedMaterials;
+	GetUsedMaterials(UsedMaterials);
+	for (const UMaterialInterface* MaterialInterface : UsedMaterials)
+	{
+		if (MaterialInterface && MaterialInterface->IsUsingWorldPositionOffset_Concurrent(GMaxRHIShaderPlatform))
+		{
+			Params.bAnyMaterialHasWorldPositionOffset = true;
+			break;
+		}
+	}
+}
+
+#endif
+
+bool FMassPrimitiveRenderStateHelper::IsPSOPrecaching() const
+{
+#if !WITH_EDITOR && UE_WITH_PSO_PRECACHING
+	// Consider as precaching when marked as required to do PSOs precaching (even if task has not been launched yet) 
+	return bPSOPrecacheRequired || (LatestPSOPrecacheJobSetCompleted != LatestPSOPrecacheJobSet);
+#else
+	return false;
+#endif
+}
+
+bool FMassPrimitiveRenderStateHelper::UsePSOPrecacheFallbackMaterial() const
+{
+#if !WITH_EDITOR && UE_WITH_PSO_PRECACHING
+	return IsPSOPrecaching() && GetPSOPrecacheProxyCreationStrategy() == EPSOPrecacheProxyCreationStrategy::UseFallbackMaterialUntilPSOPrecached;
+#else
+	return false;
+#endif
+}
+
+bool FMassPrimitiveRenderStateHelper::CheckPSOPrecachingAndBoostPriority(EPSOPrecachePriority NewPSOPrecachePriority)
+{
+#if !WITH_EDITOR && UE_WITH_PSO_PRECACHING
+	bool bPrecacheStillRunning = IsPSOPrecaching();
+
+	ensure(!IsComponentPSOPrecachingEnabled() || bPSOPrecacheCalled || bPSOPrecacheRequired);
+	check(NewPSOPrecachePriority == EPSOPrecachePriority::High || NewPSOPrecachePriority == EPSOPrecachePriority::Highest);
+
+	if (bPrecacheStillRunning && PSOPrecacheRequestPriority < NewPSOPrecachePriority)
+	{
+		// Only boost PSO priority if PSO task was started
+		if (LatestPSOPrecacheJobSetCompleted != LatestPSOPrecacheJobSet)
+		{
+			BoostPSOPriority(NewPSOPrecachePriority, MaterialPSOPrecacheRequestIDs);
+		}
+		PSOPrecacheRequestPriority = NewPSOPrecachePriority;
+	}
+	return bPrecacheStillRunning;
+#else
+	return false;
+#endif
+}
+
+
+void FMassPrimitiveRenderStateHelper::MarkPrecachePSOsRequired()
+{
+#if !WITH_EDITOR && UE_WITH_PSO_PRECACHING
+	if (bPSOPrecacheCalled || !FApp::CanEverRender() || !IsComponentPSOPrecachingEnabled())
+	{
+		return;
+	}
+	bPSOPrecacheRequired = true;
+	PSOPrecacheRequestPriority = EPSOPrecachePriority::Medium;
+#endif
+}
+
+void FMassPrimitiveRenderStateHelper::PrecachePSOs()
+{
+#if !WITH_EDITOR && UE_WITH_PSO_PRECACHING
+	TRACE_CPUPROFILER_EVENT_SCOPE(FMassPrimitiveRenderStateHelper::PrecachePSOs);
+
+	if (!bPSOPrecacheRequired)
+	{
+		return;
+	}
+
+	// Clear the current request data
+	MaterialPSOPrecacheRequestIDs.Empty();
+
+	// Collect the data from the derived classes
+	FPSOPrecacheParams PSOPrecacheParams;
+	SetupPrecachePSOParams(PSOPrecacheParams);
+	FMaterialInterfacePSOPrecacheParamsList PSOPrecacheDataArray;
+	CollectPSOPrecacheData(PSOPrecacheParams, PSOPrecacheDataArray);
+	// Set priority
+	for (FMaterialInterfacePSOPrecacheParams& Params : PSOPrecacheDataArray)
+	{
+		Params.Priority = PSOPrecacheRequestPriority;
+	}
+
+	FGraphEventArray GraphEvents;
+	PrecacheMaterialPSOs(PSOPrecacheDataArray, MaterialPSOPrecacheRequestIDs, GraphEvents);
+
+	RequestRecreateRenderStateWhenPSOPrecacheFinished(GraphEvents);
+	bPSOPrecacheRequired = false;
+#endif
+}
+
+UObject* FMassPrimitiveRenderStateHelper::GetOwner() const
+{
+	return GetEntityManager().GetOwner();
+}
+
+FString FMassPrimitiveRenderStateHelper::GetOwnerName() const
+{
+	return GetNameSafe(GetEntityManager().GetOwner());
+}
+
+FPrimitiveSceneProxy* FMassPrimitiveRenderStateHelper::CreateSceneProxy()
+{
+	return CreateSceneProxy(/*OutError*/nullptr);
+}
+
+//
+// Functions below here added for use in streamer, which is not used by FastGeo. These functions may not get used.
+static FRenderAssetOwnerStreamingState GFastGeoStreamingStateStub;
+FRenderAssetOwnerStreamingState& FMassPrimitiveRenderStateHelper::GetStreamingState() const
+{
+	return GFastGeoStreamingStateStub;
+}
+
+ULevel* FMassPrimitiveRenderStateHelper::GetComponentLevel() const
+{
+	return nullptr;
+}
+
+IPrimitiveComponent* FMassPrimitiveRenderStateHelper::GetLODParentPrimitive() const
+{
+	return nullptr;
+}
+
+float FMassPrimitiveRenderStateHelper::GetMinDrawDistance() const
+{
+	return 0.0f;
+}
+
+float FMassPrimitiveRenderStateHelper::GetStreamingScale() const
+{
+	return 1.0f;
+}
+
+void FMassPrimitiveRenderStateHelper::OnRenderAssetFirstLodChange(const UStreamableRenderAsset* RenderAsset, int32 FirstLodIndex)
+{
+}
+
+UStreamableRenderAsset* FMassPrimitiveRenderStateHelper::GetStreamableNaniteAsset() const
+{
+	return nullptr;
+}
+
+void FMassPrimitiveRenderStateHelper::GetStreamableRenderAssetInfo(TArray<struct FStreamingRenderAssetPrimitiveInfo>& StreamableRenderAssets) const
+{
+}
+
+void FMassPrimitiveRenderStateHelper::GetStreamingRenderAssetInfo(FStreamingTextureLevelContext& LevelContext, TArray<FStreamingRenderAssetPrimitiveInfo>& OutStreamingRenderAssets) const
+{
+}
+
+#if WITH_EDITOR
+HHitProxy* FMassPrimitiveRenderStateHelper::CreateMeshHitProxy(int32 SectionIndex, int32 MaterialIndex)
+{
+	if (FMassTypeElementFragment* TypeElementFragment = GetEntityManager().GetFragmentDataPtr<FMassTypeElementFragment>(EntityHandle))
+	{
+		return new HTypeElementHandleHitProxy(TypeElementFragment->TypeElementHandle, MaterialIndex, SectionIndex);
+	}
+	return nullptr;
+}
+#endif // WITH_EDITOR
+
+HHitProxy* FMassPrimitiveRenderStateHelper::CreatePrimitiveHitProxies(TArray<TRefCountPtr<HHitProxy>>& OutHitProxies)
+{
+#if WITH_EDITOR
+	if (FMassTypeElementFragment* TypeElementFragment = GetEntityManager().GetFragmentDataPtr<FMassTypeElementFragment>(EntityHandle))
+	{
+		HTypeElementHandleHitProxy* HitProxy = new HTypeElementHandleHitProxy(TypeElementFragment->TypeElementHandle);
+		OutHitProxies.Add(HitProxy);
+		return HitProxy;
+	}
+#endif // WITH_EDITOR
+	return nullptr;
+}
+
+FPrimitiveMaterialPropertyDescriptor FMassPrimitiveRenderStateHelper::GetUsedMaterialPropertyDesc(EShaderPlatform InShaderPlatform) const
+{
+	return FPrimitiveComponentHelper::GetUsedMaterialPropertyDesc(*this, InShaderPlatform);
+}
\ No newline at end of file
diff --git a/Engine/Source/Runtime/Mass/MassEngine/Private/Mesh/MassPrimitiveRenderStateHelper.h b/Engine/Source/Runtime/Mass/MassEngine/Private/Mesh/MassPrimitiveRenderStateHelper.h
new file mode 100644
--- /dev/null
+++ b/Engine/Source/Runtime/Mass/MassEngine/Private/Mesh/MassPrimitiveRenderStateHelper.h
@@ -0,0 +1,159 @@
+// Copyright Epic Games, Inc. All Rights Reserved.
+
+#pragma once
+
+#include "MassRenderStateHelper.h"
+
+#include "Components/ComponentInterfaces.h"
+#include "Components/SceneComponent.h"
+#include "Elements/Framework/TypedElementHandle.h"
+#include "Engine/EngineTypes.h"
+#include "HitProxies.h"
+#include "PipelineStateCache.h"
+#include "PrimitiveSceneDesc.h"
+#include "PrimitiveSceneInfoData.h"
+#include "SceneTypes.h"
+
+struct FPrimitiveSceneProxyDesc;
+struct FPrimitiveSceneDesc;
+struct FMassRenderPrimitiveFragment;
+class URuntimeVirtualTexture;
+enum EShaderPlatform : uint16;
+
+/**
+ * Helper to handle all the communication to the renderer from Mass for all type of primitives
+ */
+struct FMassPrimitiveRenderStateHelper : public FMassRenderStateHelper, public IPrimitiveComponent
+{
+public:
+	using Super = FMassRenderStateHelper;
+
+	FMassPrimitiveRenderStateHelper(FMassEntityHandle InEntityHandle, TNotNull<FMassEntityManager*> InEntityManager, const FMassRenderPrimitiveFragment& RenderPrimitiveFragment);
+	virtual ~FMassPrimitiveRenderStateHelper() override = default;
+	
+	TConstArrayView<URuntimeVirtualTexture*> GetRuntimeVirtualTextures() const;
+	bool IsFirstPersonRelevant() const;
+
+	FPrimitiveComponentId GetPrimitiveSceneId() const;
+	virtual UObject const* AdditionalStatObject() const 
+	{ 
+		return nullptr; 
+	}
+		
+	FMatrix GetRenderMatrix() const;
+
+	// Used by FStaticMeshComponentHelper/FInstancedStaticMeshComponentHelper
+	FTransform GetComponentTransform() const { return GetTransform(); }
+
+	// Primitive Fragments helpers
+	const FMassRenderPrimitiveFragment& GetRenderPrimitiveFragment() const;
+	FMassRenderPrimitiveFragment& GetMutableRenderPrimitiveFragment();
+
+#if WITH_EDITOR
+	virtual void NotifyRenderStateChanged();
+	void UpdateSelection();
+	void UpdateLevelInstanceEditingState();
+#endif // WITH_EDITOR
+	bool IsDrawnInGame() const;
+	bool CheckPSOPrecachingAndBoostPriority(EPSOPrecachePriority NewPSOPrecachePriority = EPSOPrecachePriority::High);
+	bool IsPSOPrecaching() const;
+	void SetCollisionEnabled(bool bInCollisionEnabled);
+	void UpdateVisibility();
+	void UpdateTransform();
+
+	EComponentMobility::Type GetMobility() const;
+
+	//~ Begin IPrimitiveComponent interface
+	virtual bool IsRenderStateDirty() const override;
+	virtual bool IsRenderStateCreated() const override;
+	virtual bool ShouldCreateRenderState() const override;
+	virtual bool IsRegistered() const override;
+	virtual bool IsUnreachable() const override;
+	virtual bool IsStaticMobility() const override;
+	virtual bool IsMipStreamingForced() const override;
+	virtual UWorld* GetWorld() const override;
+	virtual FSceneInterface* GetScene() const override;
+	virtual FPrimitiveSceneProxy* GetSceneProxy() const override;
+	virtual void MarkRenderStateDirty() override;
+	virtual void DestroyRenderState() override;
+	virtual void CreateRenderState(FRegisterComponentContext* Context) override;
+	virtual FString GetName() const override;
+	virtual FString GetFullName() const override;
+	virtual FTransform GetTransform() const override;
+	virtual const FBoxSphereBounds& GetBounds() const override;
+	virtual float GetLastRenderTimeOnScreen() const override;
+	virtual void GetPrimitiveStats(FPrimitiveStats& PrimitiveStats) const override;
+	virtual UObject* GetUObject() override;
+	virtual const UObject* GetUObject() const override;
+	virtual void PrecachePSOs() override;
+	virtual UObject* GetOwner() const override;
+	virtual FString GetOwnerName() const override;
+	virtual FPrimitiveSceneProxy* CreateSceneProxy() override;
+	virtual FRenderAssetOwnerStreamingState& GetStreamingState() const override;
+	virtual ULevel* GetComponentLevel() const override;
+	virtual IPrimitiveComponent* GetLODParentPrimitive() const override;
+	virtual float GetMinDrawDistance() const override;
+	virtual float GetStreamingScale() const override;
+	virtual void OnRenderAssetFirstLodChange(const UStreamableRenderAsset* RenderAsset, int32 FirstLodIndex) override;
+	virtual UStreamableRenderAsset* GetStreamableNaniteAsset() const override;
+	virtual void GetStreamableRenderAssetInfo(TArray<FStreamingRenderAssetPrimitiveInfo>& StreamableRenderAssets) const override;
+	virtual void GetStreamingRenderAssetInfo(FStreamingTextureLevelContext& LevelContext, TArray<FStreamingRenderAssetPrimitiveInfo>& OutStreamingRenderAssets) const override;
+#if WITH_EDITOR
+	virtual HHitProxy* CreateMeshHitProxy(int32 SectionIndex, int32 MaterialIndex) override;
+#endif // WITH_EDITOR
+	virtual HHitProxy* CreatePrimitiveHitProxies(TArray<TRefCountPtr<HHitProxy> >& OutHitProxies) override;
+	//~ End IPrimitiveComponent interface
+
+	//~ Begin FMassRenderStateHelper interface
+	virtual void DestroyRenderState(FMassDestroyRenderStateContext* Context) override;
+	virtual void MarkPrecachePSOsRequired() override;
+	//~End FMassRenderStateHelper interface
+
+	//~ Begin Materials
+	virtual UMaterialInterface* GetMaterial(int32 ElementIndex) const = 0;
+	virtual int32 GetNumMaterials() const = 0;
+	FPrimitiveMaterialPropertyDescriptor GetUsedMaterialPropertyDesc(EShaderPlatform InShaderPlatform) const;
+	//~ End Materials
+	
+protected:
+	virtual void InitializeSceneProxyDescDynamicProperties();
+	virtual void ResetSceneProxyDescUnsupportedProperties();
+	virtual FPrimitiveSceneProxyDesc& GetSceneProxyDesc() = 0;
+	virtual const FPrimitiveSceneProxyDesc& GetSceneProxyDesc() const = 0;
+	virtual FPrimitiveSceneProxy* CreateSceneProxy(ESceneProxyCreationError* OutError) = 0;
+	FPrimitiveSceneDesc BuildSceneDesc();
+
+	friend class UMassCreateRenderSceneProxyProcessor;
+private:
+	bool UsePSOPrecacheFallbackMaterial() const;
+
+#if !WITH_EDITOR && UE_WITH_PSO_PRECACHING
+public:
+	virtual void OnPrecacheFinished(int32 JobSetThatJustCompleted) override;
+
+protected:
+	void RequestRecreateRenderStateWhenPSOPrecacheFinished(const FGraphEventArray& PSOPrecacheCompileEvents);
+	void SetupPrecachePSOParams(FPSOPrecacheParams& Params);
+	virtual void CollectPSOPrecacheData(const FPSOPrecacheParams& BasePrecachePSOParams, FMaterialInterfacePSOPrecacheParamsList& OutParams) {}
+
+	// Cached array of material PSO requests which can be used to boost the priority
+	TArray<FMaterialPSOPrecacheRequestID> MaterialPSOPrecacheRequestIDs;
+	// Atomic int used to track the last PSO precache events
+	std::atomic<int> LatestPSOPrecacheJobSetCompleted = 0;
+	int32 LatestPSOPrecacheJobSet = 0;
+	// Helper flag to check if PSOs have been precached already
+	std::atomic<bool> bPSOPrecacheCalled = false;
+	std::atomic<bool> bPSOPrecacheRequired = false;
+	// PSOs requested priority
+	std::atomic<EPSOPrecachePriority> PSOPrecacheRequestPriority = EPSOPrecachePriority::Medium;
+	static_assert((int)EPSOPrecachePriority::Highest < 1 << 2);
+#endif
+
+protected:
+	// Copied data from fragment as scene proxy decs take a view on this array
+	FCustomPrimitiveData CustomPrimitiveData;
+
+	// Runtime data
+	FPrimitiveSceneDesc SceneDesc;
+	FPrimitiveSceneInfoData PrimitiveSceneData{};
+};
diff --git a/Engine/Source/Runtime/Mass/MassEngine/Private/Mesh/MassRenderStateHelper.cpp b/Engine/Source/Runtime/Mass/MassEngine/Private/Mesh/MassRenderStateHelper.cpp
new file mode 100644
--- /dev/null
+++ b/Engine/Source/Runtime/Mass/MassEngine/Private/Mesh/MassRenderStateHelper.cpp
@@ -0,0 +1,128 @@
+// Copyright Epic Games, Inc. All Rights Reserved.
+
+#include "MassRenderStateHelper.h"
+
+#include "MassEngineTypes.h"
+#include "Engine/World.h"
+#include "MassEntityManager.h"
+#include "MassSignalSubsystem.h"
+#include "Mesh/MassEngineMeshFragments.h"
+
+
+FMassRenderStateFragment::~FMassRenderStateFragment()
+{
+	DestroyRenderStateHelper();
+}
+
+// ----------------------------------------------------------------------//
+// FMassRenderStateFragment
+//----------------------------------------------------------------------//
+void FMassRenderStateFragment::DestroyRenderStateHelper()
+{
+	delete RenderStateHelper;
+	RenderStateHelper = nullptr;
+}
+
+FMassRenderStateFragment::FMassRenderStateFragment(FMassRenderStateFragment&& Other)
+	: FMassFragment(std::move(Other))
+    , RenderStateHelper(Other.RenderStateHelper)
+{
+	Other.RenderStateHelper = nullptr;
+}
+
+// ----------------------------------------------------------------------//
+// FMassRenderStateHelper
+//----------------------------------------------------------------------//
+FMassRenderStateHelper::FMassRenderStateHelper(FMassEntityHandle InEntityHandle, TNotNull<FMassEntityManager*> InEntityManager)
+	: EntityHandle(InEntityHandle)
+	, EntityManager(InEntityManager->AsShared())
+{
+}
+
+UWorld* FMassRenderStateHelper::GetWorld() const
+{
+	return GetEntityManager().GetWorld();
+}
+
+FSceneInterface* FMassRenderStateHelper::GetScene() const
+{
+	if (UWorld* World = GetWorld())
+	{
+		return World->Scene;
+	}
+	return nullptr;
+}
+
+FMassEntityManager& FMassRenderStateHelper::GetEntityManager() const
+{
+	return EntityManager.Get();
+}
+
+const FMassRenderStateFragment& FMassRenderStateHelper::GetRenderStateFragment() const
+{
+	return GetEntityManager().GetFragmentDataChecked<FMassRenderStateFragment>(EntityHandle);
+}
+
+FMassRenderStateFragment& FMassRenderStateHelper::GetMutableRenderStateFragment()
+{
+	return GetEntityManager().GetFragmentDataChecked<FMassRenderStateFragment>(EntityHandle);
+}
+
+bool FMassRenderStateHelper::IsRenderStateCreated() const
+{
+	return ProxyState == EProxyCreationState::Created;
+}
+
+bool FMassRenderStateHelper::ShouldCreateRenderState() const 
+{ 
+	return false; 
+}
+
+void FMassRenderStateHelper::CreateRenderState(FRegisterComponentContext* Context)
+{
+}
+
+void FMassRenderStateHelper::DestroyRenderState(FMassDestroyRenderStateContext* Context) 
+{
+}
+
+void FMassRenderStateHelper::UpdateVisibility() 
+{
+}
+
+void FMassRenderStateHelper::MarkPrecachePSOsRequired() 
+{
+}
+
+void FMassRenderStateHelper::PrecachePSOs() 
+{
+}
+
+#if !WITH_EDITOR && UE_WITH_PSO_PRECACHING
+void FMassRenderStateHelper::OnPrecacheFinished(int32 JobSetThatJustCompleted) 
+{
+}
+#endif // !WITH_EDITOR && UE_WITH_PSO_PRECACHING
+
+bool FMassRenderStateHelper::IsRenderStateDelayed() const
+{
+	return ProxyState == EProxyCreationState::Delayed;
+}
+
+bool FMassRenderStateHelper::IsRenderStateDirty() const
+{
+	return bRenderStateDirty;
+}
+
+void FMassRenderStateHelper::MarkRenderStateDirty()
+{
+	if ((IsRenderStateCreated() || IsRenderStateDelayed()) && !bRenderStateDirty)
+	{
+		bRenderStateDirty = true;
+
+		if (UMassSignalSubsystem* SignalSubsystem = UWorld::GetSubsystem<UMassSignalSubsystem>(GetWorld()))
+		{
+			SignalSubsystem->SignalEntity(UE::Mass::Signals::RenderStateDirty, EntityHandle);
+		}
+	}
+}
\ No newline at end of file
diff --git a/Engine/Source/Runtime/Mass/MassEngine/Private/Mesh/MassRenderStateHelper.h b/Engine/Source/Runtime/Mass/MassEngine/Private/Mesh/MassRenderStateHelper.h
new file mode 100644
--- /dev/null
+++ b/Engine/Source/Runtime/Mass/MassEngine/Private/Mesh/MassRenderStateHelper.h
@@ -0,0 +1,138 @@
+// Copyright Epic Games, Inc. All Rights Reserved.
+
+#pragma once
+
+#include "MassEntityElementTypes.h"
+#include "MassEntityHandle.h"
+#include "PSOPrecacheValidation.h"
+#include "MassExternalSubsystemTraits.h"
+
+#include "MassRenderStateHelper.generated.h"
+
+struct FMassRenderStateFragment;
+struct FMassEntityManager;
+class UActorComponent;
+class FSceneInterface;
+struct FMassRenderStateHelper;
+class FRegisterComponentContext;
+struct FMassDestroyRenderStateContext;
+
+enum class ESceneProxyCreationError
+{
+	None,
+	WaitingPSOs,
+	InvalidMesh
+};
+
+/**
+ * Fragment holding the render state
+ */
+USTRUCT()
+struct FMassRenderStateFragment : public FMassFragment
+{
+	GENERATED_BODY()
+
+	FMassRenderStateFragment() = default;
+	~FMassRenderStateFragment();
+	void DestroyRenderStateHelper();
+	FMassRenderStateFragment(const FMassRenderStateFragment& Other) = delete;
+	FMassRenderStateFragment(FMassRenderStateFragment&& Other);
+
+	FMassRenderStateFragment& operator=(const FMassRenderStateFragment& Other) = delete;
+	FMassRenderStateFragment& operator=(FMassRenderStateFragment&& Other) = delete;
+
+	template<typename T = FMassRenderStateHelper>
+	const T& GetRenderStateHelper() const requires TIsDerivedFrom<T, FMassRenderStateHelper>::Value
+	{
+		checkf(RenderStateHelper, TEXT("Expecting the ptr to be set"));
+		// @Todo find a better way of doing this unsafe cast, Instanced struct are not possible du to atomics member variables
+		return static_cast<T&>(*RenderStateHelper);
+	}
+
+	template<typename T = FMassRenderStateHelper>
+	T& GetRenderStateHelper() requires TIsDerivedFrom<T, FMassRenderStateHelper>::Value
+	{
+		checkf(RenderStateHelper, TEXT("Expecting the ptr to be set"));
+		// @Todo find a better way of doing this unsafe cast, Instanced struct are not possible du to atomics member variables
+		return static_cast<T&>(*RenderStateHelper);
+	}
+
+	template<typename T, typename... TArgs>
+	T* CreateRenderStateHelper(TArgs&&... InArgs) requires TIsDerivedFrom<T, FMassRenderStateHelper>::Value
+	{
+		checkf(!RenderStateHelper, TEXT("Render State Helper already exist"));
+		T* CreatedRenderState = new T(Forward<TArgs>(InArgs)...);
+		RenderStateHelper = CreatedRenderState;
+		return CreatedRenderState;
+	}
+
+private:
+	FMassRenderStateHelper* RenderStateHelper = nullptr;
+
+};
+
+template<>
+struct TStructOpsTypeTraits<FMassRenderStateFragment> : TStructOpsTypeTraitsBase2<FMassRenderStateFragment>
+{
+	enum
+	{
+		WithCopy = false,
+	};
+};
+
+template<>
+struct TMassFragmentTraits<FMassRenderStateFragment> final
+{
+	enum
+	{
+		AuthorAcceptsItsNotTriviallyCopyable = true
+	};
+};
+
+/**
+ * Base helper to handle all the communication to the renderer from Mass
+ */
+struct FMassRenderStateHelper
+{
+public:
+
+	FMassRenderStateHelper(FMassEntityHandle InEntityHandle, TNotNull<FMassEntityManager*> InEntityManager);
+	virtual ~FMassRenderStateHelper() = default;
+
+	void MarkRenderStateDirty();
+	bool IsRenderStateDirty() const;
+	bool IsRenderStateDelayed() const;
+	bool IsRenderStateCreated() const;
+	virtual bool ShouldCreateRenderState() const;
+	virtual void CreateRenderState(FRegisterComponentContext* Context);
+	virtual void DestroyRenderState(FMassDestroyRenderStateContext* Context);
+	virtual void UpdateVisibility();
+	virtual void MarkPrecachePSOsRequired();
+	virtual void PrecachePSOs();
+#if !WITH_EDITOR && UE_WITH_PSO_PRECACHING
+	virtual void OnPrecacheFinished(int32 JobSetThatJustCompleted);
+#endif
+
+	UWorld* GetWorld() const;
+	FSceneInterface* GetScene() const;
+
+	FMassEntityManager& GetEntityManager() const;
+
+protected:
+	FMassEntityHandle EntityHandle;
+	TSharedRef<FMassEntityManager> EntityManager;
+
+	const FMassRenderStateFragment& GetRenderStateFragment() const;
+	FMassRenderStateFragment& GetMutableRenderStateFragment();
+
+	enum class EProxyCreationState : uint8
+	{
+		None, // Constructed/Initialized
+		Creating, // Actively creating the proxy
+		Created, // Proxy is now created
+		Delayed // Proxy creation delayed (used when PSO precaching is not ready when creating proxy)
+	};
+
+	EProxyCreationState ProxyState = EProxyCreationState::None;
+	bool bRenderStateDirty = false;
+};
\ No newline at end of file
diff --git a/Engine/Source/Runtime/Mass/MassEngine/Private/Mesh/MassStaticMeshRenderStateHelper.cpp b/Engine/Source/Runtime/Mass/MassEngine/Private/Mesh/MassStaticMeshRenderStateHelper.cpp
new file mode 100644
--- /dev/null
+++ b/Engine/Source/Runtime/Mass/MassEngine/Private/Mesh/MassStaticMeshRenderStateHelper.cpp
@@ -0,0 +1,362 @@
+// Copyright Epic Games, Inc. All Rights Reserved.
+
+#include "MassStaticMeshRenderStateHelper.h"
+
+#include "Engine/StaticMesh.h"
+#include "MassEntityManager.h"
+#include "Mesh/MassEngineMeshFragments.h"
+#include "Physics/MassEnginePhysicsFragments.h"
+#include "Rendering/NaniteResourcesHelper.h"
+#include "StaticMeshComponentHelper.h"
+#include "StaticMeshResources.h"
+#include "StaticMeshSceneProxyDesc.h"
+
+#if WITH_EDITOR
+#include "ObjectCacheEventSink.h"
+#include "WorldPartition/HLOD/HLODActor.h"
+#endif
+
+// ----------------------------------------------------------------------//
+// FMassBaseStaticMeshRenderStateHelper
+//----------------------------------------------------------------------//
+FMassBaseStaticMeshRenderStateHelper::FMassBaseStaticMeshRenderStateHelper(FMassEntityHandle InEntityHandle, TNotNull<FMassEntityManager*> InEntityManager, const FMassRenderPrimitiveFragment& RenderPrimitiveFragment, const FMassOverrideMaterialsFragment* OverrideMaterialsFragment)
+	: Super(InEntityHandle, InEntityManager, RenderPrimitiveFragment, OverrideMaterialsFragment)
+{
+}
+
+const Nanite::FResources* FMassBaseStaticMeshRenderStateHelper::GetNaniteResources() const
+{
+	if (GetStaticMesh() && GetStaticMesh()->GetRenderData())
+	{
+		return GetStaticMesh()->GetRenderData()->NaniteResourcesPtr.Get();
+	}
+	return nullptr;
+}
+
+#if WITH_EDITOR
+void FMassBaseStaticMeshRenderStateHelper::NotifyRenderStateChanged()
+{
+	Super::NotifyRenderStateChanged();
+
+	FObjectCacheEventSink::NotifyStaticMeshChanged_Concurrent(this);
+}
+#endif
+
+int32 FMassBaseStaticMeshRenderStateHelper::GetNumMaterials() const
+{
+	return GetStaticMesh() ? GetStaticMesh()->GetStaticMaterials().Num() : 0;
+}
+
+UMaterialInterface* FMassBaseStaticMeshRenderStateHelper::GetMaterial(int32 MaterialIndex) const
+{
+	return GetMaterial(MaterialIndex, false);
+}
+
+const FCollisionResponseContainer& FMassBaseStaticMeshRenderStateHelper::GetCollisionResponseToChannels() const
+{
+	if (FMassColliderFragment* ColliderFragment = GetEntityManager().GetConstSharedFragmentDataPtr<FMassColliderFragment>(EntityHandle))
+	{
+		return ColliderFragment->CollisionResponse;
+	}
+
+	return FCollisionResponseContainer::GetDefaultResponseContainer();
+}
+
+UMaterialInterface* FMassBaseStaticMeshRenderStateHelper::GetNaniteAuditMaterial(int32 MaterialIndex) const
+{
+	return GetMaterial(MaterialIndex, true);
+}
+
+bool FMassBaseStaticMeshRenderStateHelper::UseNaniteOverrideMaterials(bool bDoingMaterialAudit) const
+{
+	return Nanite::FNaniteResourcesHelper::UseNaniteOverrideMaterials(*this, bDoingMaterialAudit);
+}
+
+bool FMassBaseStaticMeshRenderStateHelper::ShouldCreateNaniteProxy(Nanite::FMaterialAudit* OutNaniteMaterials) const
+{
+	return Nanite::FNaniteResourcesHelper::ShouldCreateNaniteProxy(GetStaticMeshSceneProxyDesc(), OutNaniteMaterials);
+}
+
+UMaterialInterface* FMassBaseStaticMeshRenderStateHelper::GetMaterial(int32 MaterialIndex, bool bDoingNaniteMaterialAudit) const
+{
+	return FStaticMeshComponentHelper::GetMaterial(*this, MaterialIndex, bDoingNaniteMaterialAudit);
+}
+
+bool FMassBaseStaticMeshRenderStateHelper::HasValidNaniteData() const
+{
+	return Nanite::FNaniteResourcesHelper::HasValidNaniteData(*this);
+}
+
+void FMassBaseStaticMeshRenderStateHelper::ResetSceneProxyDescUnsupportedProperties()
+{
+	Super::ResetSceneProxyDescUnsupportedProperties();
+
+	// Unsupported properties
+	FStaticMeshSceneProxyDesc& SceneProxyDesc = GetStaticMeshSceneProxyDesc();
+	SceneProxyDesc.LODData = TArrayView<struct FStaticMeshComponentLODInfo>();
+	SceneProxyDesc.LODParentPrimitive = nullptr;
+	SceneProxyDesc.MeshPaintTexture = nullptr;
+	SceneProxyDesc.MeshPaintTextureCoordinateIndex = 0;
+#if STATICMESH_ENABLE_DEBUG_RENDERING
+	SceneProxyDesc.bDrawMeshCollisionIfComplex = false;
+	SceneProxyDesc.bDrawMeshCollisionIfSimple = false;
+#endif // STATICMESH_ENABLE_DEBUG_RENDERING
+	SceneProxyDesc.bDisplayNaniteFallbackMesh = false;
+#if WITH_EDITORONLY_DATA
+	SceneProxyDesc.SectionIndexPreview = INDEX_NONE;
+	SceneProxyDesc.MaterialIndexPreview = INDEX_NONE;
+	SceneProxyDesc.SelectedEditorMaterial = INDEX_NONE;
+	SceneProxyDesc.SelectedEditorSection = INDEX_NONE;
+#endif // WITH_EDITORONLY_DATA
+
+#if WITH_EDITORONLY_DATA
+	SceneProxyDesc.MaterialStreamingRelativeBoxes = TArrayView<uint32>();
+#endif // WITH_EDITORONLY_DATA
+
+	// Properties that will be initialized by InitializeSceneProxyDescDynamicProperties
+	SceneProxyDesc.NaniteResources = nullptr;
+	SceneProxyDesc.BodySetup = nullptr;
+	SceneProxyDesc.MaterialRelevance = FMaterialRelevance();
+	SceneProxyDesc.bUseProvidedMaterialRelevance = false;
+}
+
+void FMassBaseStaticMeshRenderStateHelper::InitializeSceneProxyDescDynamicProperties()
+{
+	Super::InitializeSceneProxyDescDynamicProperties();
+
+	FStaticMeshSceneProxyDesc& SceneProxyDesc = GetStaticMeshSceneProxyDesc();
+	TConstArrayView<TObjectPtr<UMaterialInterface>> OverrideMaterials = GetOverrideMaterials();
+	SceneProxyDesc.OverrideMaterials = (TArrayView<TObjectPtr<UMaterialInterface>>&)OverrideMaterials; // ugly cast, but just making it work for now. It would mean to change the SceneProxyDesc to fix this.
+	SceneProxyDesc.NaniteResources = GetNaniteResources();
+	SceneProxyDesc.SetMaterialRelevance(GetMaterialRelevance(GetScene()->GetShaderPlatform()));
+	SceneProxyDesc.SetCollisionResponseToChannels(GetCollisionResponseToChannels());
+
+	// Add LODData support
+	// Add LODParentPrimitive support
+	// Add MeshPaintTexture/MeshPaintTextureCoordinateIndex support
+}
+
+FPrimitiveSceneProxy* FMassBaseStaticMeshRenderStateHelper::CreateSceneProxy(ESceneProxyCreationError* OutError)
+{
+	check(GetWorld());
+	FSceneInterface* Scene = GetScene();
+	UStaticMesh* StaticMesh = GetStaticMesh();
+	check(Scene);
+	check(StaticMesh);
+	check(StaticMesh->GetRenderData());
+	check(StaticMesh->GetRenderData()->IsInitialized());
+	check(!StaticMesh->IsCompiling());
+
+	if (OutError)
+	{
+		*OutError = ESceneProxyCreationError::None;
+	}
+
+	InitializeSceneProxyDescDynamicProperties();
+
+	const FStaticMeshSceneProxyDesc& SceneProxyDesc = GetStaticMeshSceneProxyDesc();
+	check(SceneProxyDesc.Scene);
+	check(SceneProxyDesc.Scene == Scene);
+	check(SceneProxyDesc.World == Scene->GetWorld());
+	check(SceneProxyDesc.FeatureLevel == Scene->GetFeatureLevel());
+	check(SceneProxyDesc.ComponentId == GetPrimitiveSceneId());
+
+	FStaticMeshComponentHelper::ESceneProxyCreationError SceneProxyCreationError;
+	FPrimitiveSceneProxy* SceneProxy = FStaticMeshComponentHelper::CreateSceneProxy<FMassBaseStaticMeshRenderStateHelper, /*bRenderDataReady=*/true>(*this, &SceneProxyCreationError);
+
+	if (SceneProxy == nullptr && OutError)
+	{
+		switch (SceneProxyCreationError)
+		{
+		case FStaticMeshComponentHelper::ESceneProxyCreationError::WaitingPSOs:
+			*OutError = ESceneProxyCreationError::WaitingPSOs;
+			break;
+
+		default:
+			*OutError = ESceneProxyCreationError::InvalidMesh;
+			break;
+		}
+	}
+
+	return SceneProxy;
+}
+
+
+FPrimitiveSceneProxy* FMassBaseStaticMeshRenderStateHelper::CreateStaticMeshSceneProxy(const Nanite::FMaterialAudit& NaniteMaterials, bool bCreateNanite)
+{
+	check(GetWorld());
+	const FStaticMeshSceneProxyDesc& SceneProxyDesc = GetStaticMeshSceneProxyDesc();
+	check(SceneProxyDesc.Scene);
+
+	if (bCreateNanite)
+	{
+		PrimitiveSceneData.SceneProxy = ::new Nanite::FSceneProxy(NaniteMaterials, SceneProxyDesc);
+	}
+	else
+	{
+		PrimitiveSceneData.SceneProxy = ::new FStaticMeshSceneProxy(SceneProxyDesc, false);
+	}
+	return PrimitiveSceneData.SceneProxy;
+}
+
+void FMassBaseStaticMeshRenderStateHelper::GetUsedMaterials(TArray<UMaterialInterface*>& OutMaterials, bool bGetDebugMaterials) const
+{
+	FStaticMeshComponentHelper::GetUsedMaterials(*this, OutMaterials, bGetDebugMaterials);
+}
+
+void FMassBaseStaticMeshRenderStateHelper::GetDefaultMaterialSlotsOverlayMaterial(TArray<TObjectPtr<UMaterialInterface>>& OutMaterialSlotOverlayMaterials) const
+{
+	OutMaterialSlotOverlayMaterials.Reset();
+	if (GetStaticMesh())
+	{
+		const TArray<FStaticMaterial>& StaticMaterials = GetStaticMesh()->GetStaticMaterials();
+		for (const FStaticMaterial& StaticMaterial : StaticMaterials)
+		{
+			OutMaterialSlotOverlayMaterials.Add(StaticMaterial.OverlayMaterialInterface);
+		}
+	}
+}
+
+const TArray<TObjectPtr<UMaterialInterface>>& FMassBaseStaticMeshRenderStateHelper::GetComponentMaterialSlotsOverlayMaterial() const
+{
+	return GetStaticMeshSceneProxyDesc().MaterialSlotsOverlayMaterial;
+}
+
+#if WITH_EDITOR
+void FMassBaseStaticMeshRenderStateHelper::PostStaticMeshCompilation()
+{
+	MarkRenderStateDirty();
+}
+#endif // WITH_EDITOR
+
+UStaticMesh* FMassBaseStaticMeshRenderStateHelper::GetStaticMesh() const
+{
+	return GetStaticMeshSceneProxyDesc().StaticMesh;
+}
+
+IPrimitiveComponent* FMassBaseStaticMeshRenderStateHelper::GetPrimitiveComponentInterface()
+{
+	return this;
+}
+
+UObject const* FMassBaseStaticMeshRenderStateHelper::AdditionalStatObject() const
+{
+	return GetStaticMesh();
+}
+
+UMaterialInterface* FMassBaseStaticMeshRenderStateHelper::GetOverlayMaterial() const
+{
+	return GetStaticMeshSceneProxyDesc().OverlayMaterial;
+}
+
+bool FMassBaseStaticMeshRenderStateHelper::IsReverseCulling() const
+{
+	return GetStaticMeshSceneProxyDesc().IsReverseCulling();
+}
+
+bool FMassBaseStaticMeshRenderStateHelper::IsDisallowNanite() const
+{
+	return GetStaticMeshSceneProxyDesc().IsDisallowNanite();
+}
+
+bool FMassBaseStaticMeshRenderStateHelper::IsForceDisableNanite() const
+{
+	return GetStaticMeshSceneProxyDesc().IsForceDisableNanite();
+}
+
+bool FMassBaseStaticMeshRenderStateHelper::IsForceNaniteForMasked() const
+{
+	return GetStaticMeshSceneProxyDesc().IsForceNaniteForMasked();
+}
+
+int32 FMassBaseStaticMeshRenderStateHelper::GetForcedLodModel() const
+{
+	return GetStaticMeshSceneProxyDesc().GetForcedLodModel();
+}
+
+bool FMassBaseStaticMeshRenderStateHelper::GetOverrideMinLOD() const
+{
+	return GetStaticMeshSceneProxyDesc().bOverrideMinLOD;
+}
+
+int32 FMassBaseStaticMeshRenderStateHelper::GetMinLOD() const
+{
+	return GetStaticMeshSceneProxyDesc().MinLOD;
+}
+
+bool FMassBaseStaticMeshRenderStateHelper::GetForceNaniteForMasked() const
+{
+	return GetStaticMeshSceneProxyDesc().bForceNaniteForMasked;
+}
+
+int32 FMassBaseStaticMeshRenderStateHelper::GetWorldPositionOffsetDisableDistance() const
+{
+	return GetStaticMeshSceneProxyDesc().GetWorldPositionOffsetDisableDistance();
+}
+
+float FMassBaseStaticMeshRenderStateHelper::GetNanitePixelProgrammableDistance() const
+{
+	return GetStaticMeshSceneProxyDesc().GetNanitePixelProgrammableDistance();
+}
+
+bool FMassBaseStaticMeshRenderStateHelper::EvaluateWorldPositionOffsetInRayTracing() const
+{
+	return GetStaticMeshSceneProxyDesc().EvaluateWorldPositionOffsetInRayTracing();
+}
+
+#if WITH_EDITORONLY_DATA
+bool FMassBaseStaticMeshRenderStateHelper::IsDisplayNaniteFallbackMesh() const
+{
+	return GetStaticMeshSceneProxyDesc().IsDisplayNaniteFallbackMesh();
+}
+#endif
+
+
+#if !WITH_EDITOR && UE_WITH_PSO_PRECACHING
+void FMassBaseStaticMeshRenderStateHelper::CollectPSOPrecacheData(const FPSOPrecacheParams& BasePrecachePSOParams, FMaterialInterfacePSOPrecacheParamsList& OutParams)
+{
+	FStaticMeshComponentHelper::CollectPSOPrecacheData(*this, BasePrecachePSOParams, OutParams);
+}
+#endif
+
+// ----------------------------------------------------------------------//
+// FMassStaticMeshRenderStateHelper
+//----------------------------------------------------------------------//
+FMassStaticMeshRenderStateHelper::FMassStaticMeshRenderStateHelper(FMassEntityHandle InEntityHandle, TNotNull<FMassEntityManager*> InEntityManager, const FMassRenderPrimitiveFragment& RenderPrimitiveFragment, const FMassOverrideMaterialsFragment* OverrideMaterialsFragment, const FMassRenderStaticMeshFragment& RenderStaticMeshFragment)
+	: Super(InEntityHandle, InEntityManager, RenderPrimitiveFragment, OverrideMaterialsFragment)
+{
+}
+
+const FMassRenderStaticMeshFragment& FMassStaticMeshRenderStateHelper::GetRenderStaticMeshFragment() const
+{
+	return GetEntityManager().GetFragmentDataChecked<FMassRenderStaticMeshFragment>(EntityHandle);
+}
+
+FMassRenderStaticMeshFragment& FMassStaticMeshRenderStateHelper::GetMutableRenderStaticMeshFragment()
+{
+	return GetEntityManager().GetFragmentDataChecked<FMassRenderStaticMeshFragment>(EntityHandle);
+}
+
+FPrimitiveSceneProxyDesc& FMassStaticMeshRenderStateHelper::GetSceneProxyDesc()
+{
+	return GetStaticMeshSceneProxyDesc();
+}
+
+const FPrimitiveSceneProxyDesc& FMassStaticMeshRenderStateHelper::GetSceneProxyDesc() const
+{
+	return GetStaticMeshSceneProxyDesc();
+}
+
+FStaticMeshSceneProxyDesc& FMassStaticMeshRenderStateHelper::GetStaticMeshSceneProxyDesc()
+{
+	FMassRenderStaticMeshFragment& RenderStaticMeshFragment = GetMutableRenderStaticMeshFragment();
+	checkf(RenderStaticMeshFragment.StaticMeshSceneProxyDesc, TEXT("Expecting a valid scene proxy desc"));
+	return *RenderStaticMeshFragment.StaticMeshSceneProxyDesc;
+}
+
+const FStaticMeshSceneProxyDesc& FMassStaticMeshRenderStateHelper::GetStaticMeshSceneProxyDesc() const
+{
+	const FMassRenderStaticMeshFragment& RenderStaticMeshFragment = GetRenderStaticMeshFragment();
+	checkf(RenderStaticMeshFragment.StaticMeshSceneProxyDesc, TEXT("Expecting a valid scene proxy desc"));
+	return *RenderStaticMeshFragment.StaticMeshSceneProxyDesc;
+}
diff --git a/Engine/Source/Runtime/Mass/MassEngine/Private/Mesh/MassStaticMeshRenderStateHelper.h b/Engine/Source/Runtime/Mass/MassEngine/Private/Mesh/MassStaticMeshRenderStateHelper.h
new file mode 100644
--- /dev/null
+++ b/Engine/Source/Runtime/Mass/MassEngine/Private/Mesh/MassStaticMeshRenderStateHelper.h
@@ -0,0 +1,121 @@
+// Copyright Epic Games, Inc. All Rights Reserved.
+
+#pragma once
+
+#include "MassMeshRenderStateHelper.h"
+
+class FColorVertexBuffer;
+class UMaterialInterface;
+struct FStaticMeshSceneProxyDesc;
+struct FMassRenderStaticMeshFragment;
+namespace Nanite
+{
+	struct FResources;
+	struct FMaterialAudit;
+}
+
+/**
+ * Base helper to handle all the communication to the renderer from Mass for all type of static meshes
+ */
+struct FMassBaseStaticMeshRenderStateHelper : public FMassMeshRenderStateHelper, public IStaticMeshComponent
+{
+public:
+	using Super = FMassMeshRenderStateHelper;
+
+	FMassBaseStaticMeshRenderStateHelper(FMassEntityHandle InEntityHandle, TNotNull<FMassEntityManager*> InEntityManager, const FMassRenderPrimitiveFragment& RenderPrimitiveFragment, const FMassOverrideMaterialsFragment* OverrideMaterialsFragment);
+	virtual ~FMassBaseStaticMeshRenderStateHelper() override = default;
+
+	const FCollisionResponseContainer& GetCollisionResponseToChannels() const;
+	UMaterialInterface* GetNaniteAuditMaterial(int32 MaterialIndex) const;
+	const Nanite::FResources* GetNaniteResources() const;
+	bool HasValidNaniteData() const;
+	bool UseNaniteOverrideMaterials(bool bDoingMaterialAudit) const;
+	bool ShouldCreateNaniteProxy(Nanite::FMaterialAudit* OutNaniteMaterials) const;
+	virtual UObject const* AdditionalStatObject() const override;
+	bool IsReverseCulling() const;
+	bool IsDisallowNanite() const;
+	bool IsForceDisableNanite() const;
+	bool IsForceNaniteForMasked() const;
+	int32 GetForcedLodModel() const;
+	bool GetOverrideMinLOD() const;
+	int32 GetMinLOD() const;
+	bool GetForceNaniteForMasked() const;
+	int32 GetWorldPositionOffsetDisableDistance() const;
+	float GetNanitePixelProgrammableDistance() const;
+	bool EvaluateWorldPositionOffsetInRayTracing() const;
+#if WITH_EDITORONLY_DATA
+	bool IsDisplayNaniteFallbackMesh() const;
+#endif
+
+	//~ Begin IStaticMeshComponent interface
+#if WITH_EDITOR
+	virtual void OnMeshRebuild(bool bRenderDataChanged) override {}
+	virtual void PreStaticMeshCompilation() override {}
+	virtual void PostStaticMeshCompilation() override;
+#endif
+	virtual UStaticMesh* GetStaticMesh() const override;
+	virtual IPrimitiveComponent* GetPrimitiveComponentInterface() override;
+	//~ End IStaticMeshComponent interface
+
+	//~ Begin FMassPrimitiveRenderStateHelper interface
+	virtual UMaterialInterface* GetMaterial(int32 MaterialIndex) const override;
+	virtual int32 GetNumMaterials() const override;
+	virtual void GetUsedMaterials(TArray<UMaterialInterface*>& OutMaterials, bool bGetDebugMaterials = false) const override;
+#if WITH_EDITOR
+	virtual void NotifyRenderStateChanged() override;
+#endif
+	//~ End FMassPrimitiveRenderStateHelper interface
+
+protected:
+	//~ Begin FMassPrimitiveRenderStateHelper interface
+	virtual void ResetSceneProxyDescUnsupportedProperties() override;
+	virtual void InitializeSceneProxyDescDynamicProperties() override;
+	virtual FPrimitiveSceneProxy* CreateSceneProxy(ESceneProxyCreationError* OutError = nullptr) override;
+	//~ End FMassPrimitiveRenderStateHelper interface
+
+	//~ Begin FMassMeshRenderStateHelper interface
+	virtual UMaterialInterface* GetOverlayMaterial() const override;
+	virtual const TArray<TObjectPtr<UMaterialInterface>>& GetComponentMaterialSlotsOverlayMaterial() const override;
+	virtual void GetDefaultMaterialSlotsOverlayMaterial(TArray<TObjectPtr<UMaterialInterface>>& OutMaterialSlotOverlayMaterials) const override;
+	//~ End FMassMeshRenderStateHelper interface
+
+	virtual FStaticMeshSceneProxyDesc& GetStaticMeshSceneProxyDesc() = 0;
+	virtual const FStaticMeshSceneProxyDesc& GetStaticMeshSceneProxyDesc() const = 0;
+	virtual FPrimitiveSceneProxy* CreateStaticMeshSceneProxy(const Nanite::FMaterialAudit& NaniteMaterials, bool bCreateNanite);
+	UMaterialInterface* GetMaterial(int32 MaterialIndex, bool bDoingNaniteMaterialAudit) const;
+#if !WITH_EDITOR && UE_WITH_PSO_PRECACHING
+	virtual void CollectPSOPrecacheData(const FPSOPrecacheParams& BasePrecachePSOParams, FMaterialInterfacePSOPrecacheParamsList& OutParams);
+#endif
+
+	friend class FMeshComponentHelper;
+	friend class FPrimitiveComponentHelper;
+	friend class FStaticMeshComponentHelper;
+	friend class FInstancedStaticMeshComponentHelper;
+};
+
+/**
+ * Helper to handle all the communication to the renderer from Mass for all type of static meshes
+ */
+struct FMassStaticMeshRenderStateHelper : public FMassBaseStaticMeshRenderStateHelper
+{
+public:
+	using Super = FMassBaseStaticMeshRenderStateHelper;
+
+	FMassStaticMeshRenderStateHelper(FMassEntityHandle InEntityHandle, TNotNull<FMassEntityManager*> InEntityManager, const FMassRenderPrimitiveFragment& RenderPrimitiveFragment, const FMassOverrideMaterialsFragment* OverrideMaterialsFragment, const FMassRenderStaticMeshFragment& RenderStaticMeshFragment);
+	virtual ~FMassStaticMeshRenderStateHelper() = default;
+
+	// StaticMesh Fragments helpers
+	const FMassRenderStaticMeshFragment& GetRenderStaticMeshFragment() const;
+	FMassRenderStaticMeshFragment& GetMutableRenderStaticMeshFragment();
+
+protected:
+	//~ Begin FMassPrimitiveRenderStateHelper interface
+	virtual FPrimitiveSceneProxyDesc& GetSceneProxyDesc() override;
+	virtual const FPrimitiveSceneProxyDesc& GetSceneProxyDesc() const override;
+	//~ End FMassPrimitiveRenderStateHelper interface
+
+	//~ Begin FMassBaseStaticMeshRenderStateHelper interface
+	virtual FStaticMeshSceneProxyDesc& GetStaticMeshSceneProxyDesc() override;
+	virtual const FStaticMeshSceneProxyDesc& GetStaticMeshSceneProxyDesc() const override;
+	//~ End FMassBaseStaticMeshRenderStateHelper interface
+};
\ No newline at end of file
diff --git a/Engine/Source/Runtime/Mass/MassEngine/Public/MassEngineTypes.h b/Engine/Source/Runtime/Mass/MassEngine/Public/MassEngineTypes.h
--- a/Engine/Source/Runtime/Mass/MassEngine/Public/MassEngineTypes.h
+++ b/Engine/Source/Runtime/Mass/MassEngine/Public/MassEngineTypes.h
@@ -11,4 +11,17 @@ namespace UE::Mass::Signals
 {
 	const FName TransformChanged = FName(TEXT("TransformChanged"));
 	const FName MeshChanged = FName(TEXT("MeshChanged"));
+	const FName MeshVisualPropertyChanged = FName(TEXT("MeshVisualPropertyChanged"));
+	const FName RenderStateDirty = FName(TEXT("RenderStateDirty"));
+	
+#if WITH_EDITOR
+	const FName SelectionChanged = FName(TEXT("SelectionChanged"));
+	const FName LevelEditingStateChanged = FName(TEXT("LevelEditingStateChanged"));
+#endif // WITH_EDITOR
 }
+
+namespace UE::Mass::ProcessorGroupNames
+{
+	const FName SetupRenderState = FName(TEXT("SetupRenderState"));
+	const FName CreateRenderState = FName(TEXT("CreateRenderState"));
+}
\ No newline at end of file
diff --git a/Engine/Source/Runtime/Mass/MassEngine/Public/Mesh/MassEngineMeshFragments.h b/Engine/Source/Runtime/Mass/MassEngine/Public/Mesh/MassEngineMeshFragments.h
--- a/Engine/Source/Runtime/Mass/MassEngine/Public/Mesh/MassEngineMeshFragments.h
+++ b/Engine/Source/Runtime/Mass/MassEngine/Public/Mesh/MassEngineMeshFragments.h
@@ -2,15 +2,25 @@
 
 #pragma once
 
-#include "MassEntityElementTypes.h"
+#include "Components/InstancedStaticMeshComponent.h"
 #include "Engine/StaticMesh.h"
+#include "MassEntityElementTypes.h"
+#include "MassEntityLinkFragments.h"
+#include "MassExternalSubsystemTraits.h"
+#include "Elements/Framework/TypedElementHandle.h"
 
 #include "MassEngineMeshFragments.generated.h"
 
+struct FTransformFragment;
+class URuntimeVirtualTexture;
+struct FStaticMeshSceneProxyDesc;
+struct FInstancedStaticMeshSceneProxyDesc;
+struct FMassRenderStateHelper;
+
 #define UE_API MASSENGINE_API
 
 /**
- * Fragment holding a pointer to the static mesh associated to the entity
+ * Const shared fragment holding a pointer to the static mesh associated to the entity
  */
 USTRUCT()
 struct FMassStaticMeshFragment : public FMassConstSharedFragment
@@ -22,9 +32,316 @@ struct FMassStaticMeshFragment : public FMassConstSharedFragment
 
 	UPROPERTY(Transient, VisibleAnywhere, Category=Debug)
 	TWeakObjectPtr<const UStaticMesh> Mesh;
+};
+
+/**
+ * Const shared fragment holding all the visualization flags for a mesh
+ */
+USTRUCT()
+struct FMassVisualizationMeshFragment : public FMassConstSharedFragment
+{
+	GENERATED_BODY()
+
+	UPROPERTY(Transient, VisibleAnywhere, Category = Mesh)
+	int64 CustomDepthStencilValue = 0;
+	UPROPERTY(Transient, VisibleAnywhere, Category = Mesh)
+	int32 TranslucencySortPriority = 0;
+	UPROPERTY(Transient, VisibleAnywhere, Category = Mesh)
+	uint8 HiddenInGame : 1 = false;
+	UPROPERTY(Transient, VisibleAnywhere, Category = Mesh)
+	uint8 CastShadow : 1 = true;
+	UPROPERTY(Transient, VisibleAnywhere, Category = Mesh)
+	uint8 CastShadowAsTwoSided : 1 = false;
+	UPROPERTY(Transient, VisibleAnywhere, Category = Mesh)
+	uint8 CastHiddenShadow : 1 = false;
+	UPROPERTY(Transient, VisibleAnywhere, Category = Mesh)
+	uint8 AffectDynamicIndirectLighting : 1 = true;
+	UPROPERTY(Transient, VisibleAnywhere, Category = Mesh)
+	uint8 AffectIndirectLightingWhileHidden : 1 = false;
+	UPROPERTY(Transient, VisibleAnywhere, Category = Mesh)
+	uint8 ReceivesDecals : 1 = true;
+	UPROPERTY(Transient, VisibleAnywhere, Category = Mesh)
+	uint8 NeverDistanceCull : 1 = false;
+	UPROPERTY(Transient, VisibleAnywhere, Category = Mesh)
+	uint8 CastFarShadow : 1 = false;
+	UPROPERTY(Transient, VisibleAnywhere, Category = Mesh)
+	uint8 CastInsetShadow : 1 = false;
+	UPROPERTY(Transient, VisibleAnywhere, Category = Mesh)
+	uint8 CastContactShadow : 1 = false;
+	UPROPERTY(Transient, VisibleAnywhere, Category = Mesh)
+	uint8 RenderCustomDepth : 1 = false;
+};
+
+/**
+ * Const shared fragment holding the editor visualization flags for a mesh
+ */
+USTRUCT()
+struct FMassEditorVisualizationMeshFragment : public FMassConstSharedFragment
+{
+	GENERATED_BODY()
+
+#if  WITH_EDITORONLY_DATA
+	UPROPERTY(Transient, VisibleAnywhere, Category = Mesh)
+	uint8 bShouldRenderSelected : 1 = false;
+	UPROPERTY(Transient, VisibleAnywhere, Category = Mesh)
+	uint8 bHiddenInEditor : 1 = false;
+	UPROPERTY(Transient, VisibleAnywhere, Category = Mesh)
+	uint8 bLevelInstanceEditingState : 1 = false;
+#endif //  WITH_EDITORONLY_DATA
+};
+
+/**
+ * Const shared fragment holding the list of material override for a mesh
+ */
+USTRUCT()
+struct FMassOverrideMaterialsFragment : public FMassConstSharedFragment
+{
+	GENERATED_BODY()
+
+	UPROPERTY(Transient, VisibleAnywhere, Category = Debug)
+	TArray<TObjectPtr<UMaterialInterface>> OverrideMaterials;
+};
+
+template<>
+struct TMassFragmentTraits<FMassVisualizationMeshFragment> final
+{
+	enum
+	{
+		AuthorAcceptsItsNotTriviallyCopyable = true
+	};
+};
+
+/**
+ * Fragment holding rendering information of all the different primitives
+ */
+USTRUCT()
+struct FMassRenderPrimitiveFragment : public FMassFragment
+{
+	GENERATED_BODY()
 
 	UPROPERTY(Transient, VisibleAnywhere, Category = Debug)
-	TArray<TWeakObjectPtr<UMaterialInterface>> OverrideMaterials;
+	FBoxSphereBounds LocalBounds = {};
+	UPROPERTY(Transient, VisibleAnywhere, Category = Debug)
+	FBoxSphereBounds WorldBounds = {}; // Maybe not here...
+	UPROPERTY(Transient, VisibleAnywhere, Category = Debug)
+	FCustomPrimitiveData CustomPrimitiveData;
+	UPROPERTY(Transient, VisibleAnywhere, Category = Debug)
+	TEnumAsByte<enum EDetailMode> DetailMode = EDetailMode::DM_Low;
+	UPROPERTY(Transient, VisibleAnywhere, Category = Debug)
+	uint8 bIsVisible : 1 = true;
+	UPROPERTY(Transient, VisibleAnywhere, Category = Debug)
+	uint8 bRasterizeAsFilledConvexVolume : 1 = false;
+	UPROPERTY(Transient, VisibleAnywhere, Category = Debug)
+	TArray<TObjectPtr<URuntimeVirtualTexture>> RuntimeVirtualTextures;
+};
+
+template<>
+struct TMassFragmentTraits<FMassRenderPrimitiveFragment> final
+{
+	enum
+	{
+		AuthorAcceptsItsNotTriviallyCopyable = true
+	};
+};
+
+
+/**
+ * Fragment holding the rendering information for a static mesh
+ */
+USTRUCT()
+struct FMassRenderStaticMeshFragment : public FMassFragment
+{
+	GENERATED_BODY()
+
+	UE_API FMassRenderStaticMeshFragment();
+	UE_API explicit FMassRenderStaticMeshFragment(TNotNull<FStaticMeshSceneProxyDesc*> InStaticMeshSceneProxyDesc);
+	UE_API ~FMassRenderStaticMeshFragment();
+	FMassRenderStaticMeshFragment(const FMassRenderStaticMeshFragment& Other) = delete;
+	FMassRenderStaticMeshFragment(FMassRenderStaticMeshFragment&& Other);
+
+	FMassRenderStaticMeshFragment& operator=(const FMassRenderStaticMeshFragment& Other) = delete;
+	FMassRenderStaticMeshFragment& operator=(FMassRenderStaticMeshFragment&& Other) = delete;
+
+	UStaticMesh* GetStaticMesh();
+	const UStaticMesh* GetStaticMesh() const;
+
+	FStaticMeshSceneProxyDesc* StaticMeshSceneProxyDesc = nullptr;
+	bool bOwnProxyDesc : 1 = false;
+};
+
+template<>
+struct TStructOpsTypeTraits<FMassRenderStaticMeshFragment> : TStructOpsTypeTraitsBase2<FMassRenderStaticMeshFragment>
+{
+	enum
+	{
+		WithCopy = false,
+	};
+};
+
+template<>
+struct TMassFragmentTraits<FMassRenderStaticMeshFragment> final
+{
+	enum
+	{
+		AuthorAcceptsItsNotTriviallyCopyable = true
+	};
+};
+
+#if WITH_EDITOR
+/**
+ * Hit proxy to use along with the FMassTypeElementFragment
+ */
+struct HTypeElementHandleHitProxy : public HHitProxy
+{
+	DECLARE_HIT_PROXY();
+
+	HTypeElementHandleHitProxy(const FTypedElementHandle& InElementHandle) :
+		HHitProxy(HPP_World),
+		ElementHandle(InElementHandle)
+	{
+	}
+
+	HTypeElementHandleHitProxy(const FTypedElementHandle& InElementHandle, int32 InMaterialIndex, int32 InSectionIndex) :
+		HHitProxy(HPP_World),
+		ElementHandle(InElementHandle),
+		SectionIndex(InSectionIndex),
+		MaterialIndex(InMaterialIndex)
+	{
+	}
+
+	FTypedElementHandle ElementHandle;
+	int32 SectionIndex = -1;
+	int32 MaterialIndex = -1;
+
+	virtual bool AlwaysAllowsTranslucentPrimitives() const override
+	{
+		return false;
+	}
+
+	virtual FTypedElementHandle GetElementHandle() const override
+	{
+		return ElementHandle;
+	}
+
+	virtual EMouseCursor::Type GetMouseCursor() override;
+};
+#endif // WITH_EDITOR
+
+/**
+ * Fragment holding rendering information for a instanced static mesh
+ */
+USTRUCT()
+struct FMassRenderISMFragment : public FMassFragment
+{
+	GENERATED_BODY()
+
+	FMassRenderISMFragment();
+	explicit FMassRenderISMFragment(TNotNull<FInstancedStaticMeshSceneProxyDesc*> InInstancedStaticMeshSceneProxyDesc);
+	~FMassRenderISMFragment();
+
+	FMassRenderISMFragment(const FMassRenderISMFragment& Other) = delete;
+	FMassRenderISMFragment(FMassRenderISMFragment&& Other);
+
+	FMassRenderISMFragment& operator=(const FMassRenderISMFragment& Other) = delete;
+	FMassRenderISMFragment& operator=(FMassRenderISMFragment&& Other) = delete;
+
+	//FMassRenderISMFragment(FMassRenderISMFragment& Other);
+
+	UStaticMesh* GetStaticMesh();
+	const UStaticMesh* GetStaticMesh() const;
+
+#if WITH_EDITOR
+	TSparseArray<TRefCountPtr<HTypeElementHandleHitProxy>> PerInstanceHitProxy;
+#endif // WITH_EDITOR
+	TSparseArray<FInstancedStaticMeshInstanceData> PerInstanceSMData;
+	TArray<float> PerInstanceSMCustomData;
+	TArray<FInstancedStaticMeshRandomSeed> AdditionalRandomSeeds;
+
+	FInstancedStaticMeshSceneProxyDesc* InstancedStaticMeshSceneProxyDesc = nullptr;
+	int32 InstancingRandomSeed = 0;
+	uint8 bOwnProxyDesc : 1 = false;
+};
+
+template<>
+struct TStructOpsTypeTraits<FMassRenderISMFragment> : TStructOpsTypeTraitsBase2<FMassRenderISMFragment>
+{
+	enum
+	{
+		WithCopy = false,
+	};
+};
+
+template<>
+struct TMassFragmentTraits<FMassRenderISMFragment> final
+{
+	enum
+	{
+		AuthorAcceptsItsNotTriviallyCopyable = true
+	};
+};
+
+
+/**
+ * Fragment holding the instanced static mesh entity for the instances
+ */
+USTRUCT()
+struct FMassRenderISMLinkFragment : public FMassEntityLinkFragment
+{
+	GENERATED_BODY()
+};
+
+/**
+ * Fragment holding the index for the instances of an instanced static mesh 
+ */
+USTRUCT()
+struct FMassRenderISMInstanceFragment : public FMassFragment
+{
+	GENERATED_BODY()
+
+	UPROPERTY(Transient, VisibleAnywhere, Category = Debug)
+	int32 InstanceIndex = -1;
+#if WITH_EDITORONLY_DATA
+	UPROPERTY(Transient, VisibleAnywhere, Category = Debug)
+	int32 HitProxyIndex = -1;
+#endif // WITH_EDITORONLY_DATA
+};
+
+/**
+ * Fragment hold the type element for selection
+ */
+USTRUCT()
+struct FMassTypeElementFragment : public FMassFragment
+{
+	GENERATED_BODY()
+
+	FTypedElementHandle TypeElementHandle;
+};
+
+template<>
+struct TMassFragmentTraits<FMassTypeElementFragment> final
+{
+	enum
+	{
+		AuthorAcceptsItsNotTriviallyCopyable = true
+	};
+};
+
+/**
+ * Tag to tell if the SceneProxy was created
+ */
+USTRUCT()
+struct FMassSceneProxyCreatedTag : public FMassTag
+{
+	GENERATED_BODY()
+};
+
+/**
+ * Tag to tell if the static mesh is a candidate for instantiation
+ */
+USTRUCT()
+struct FMassRenderISMCandidateTag : public FMassTag
+{
+	GENERATED_BODY()
 };
 
-#undef UE_API
\ No newline at end of file
+#undef UE_API
diff --git a/Engine/Source/Runtime/Mass/MassEngine/Public/Mesh/MassEngineMeshUtils.h b/Engine/Source/Runtime/Mass/MassEngine/Public/Mesh/MassEngineMeshUtils.h
new file mode 100644
--- /dev/null
+++ b/Engine/Source/Runtime/Mass/MassEngine/Public/Mesh/MassEngineMeshUtils.h
@@ -0,0 +1,95 @@
+// Copyright Epic Games, Inc. All Rights Reserved.
+
+#pragma once
+
+#include "Misc/NotNull.h"
+
+struct FPrimitiveSceneProxyDesc;
+struct FInstancedStaticMeshSceneProxyDesc;
+struct FMassEditorVisualizationMeshFragment;
+struct FMassVisualizationMeshFragment;
+struct FMassStaticMeshFragment;
+struct FStaticMeshSceneProxyDesc;
+class UActorComponent;
+struct FTransformFragment;
+struct FMassRenderPrimitiveFragment;
+struct FMassRenderStaticMeshFragment;
+struct FMassRenderISMFragment;
+
+#define UE_API MASSENGINE_API
+
+namespace UE::MassEngine::Mesh
+{
+	/**
+	 * Initializes the transform, primitive fragments from a UActorComponent
+	 * @param Component to initialize the fragments from
+	 * @param TransformFragment to set up
+	 * @param RenderPrimitiveFragment to set up
+	 */
+	UE_API void InitializePrimitiveFragmentsFromComponent(TNotNull<const UActorComponent*> Component, FTransformFragment& TransformFragment, FMassRenderPrimitiveFragment& RenderPrimitiveFragment);
+
+	/**
+	 * Initializes the render transform, primitive, render static mesh fragments from a UActorComponent
+	 * @param Component to initialize the fragments from
+	 * @param TransformFragment the transform fragment to set up
+	 * @param RenderPrimitiveFragment the render primitive fragment to set up
+	 * @param RenderStaticMeshFragment the render static mesh fragment to set up
+	 */
+	UE_API void InitializeStaticMeshFragmentsFromComponent(TNotNull<const UActorComponent*> Component, FTransformFragment& TransformFragment, FMassRenderPrimitiveFragment& RenderPrimitiveFragment, FMassRenderStaticMeshFragment& RenderStaticMeshFragment);
+
+	/**
+	 * Initializes the render transform, primitive, render instanced static mesh fragments from a UActorComponent
+	 * @param Component to initialize the fragments from
+	 * @param TransformFragment the transform fragment to set up
+	 * @param RenderPrimitiveFragment the render primitive fragment to set up
+	 * @param RenderISMFragment the render instanced static mesh fragment to set up
+	 */
+	UE_API void InitializeISMFragmentsFromComponent(TNotNull<const UActorComponent*> Component, FTransformFragment& TransformFragment, FMassRenderPrimitiveFragment& RenderPrimitiveFragment, FMassRenderISMFragment& RenderISMFragment);
+
+	enum class EBoundsType
+	{
+		LocalBounds,
+		WorldBounds,
+		NavigationBounds
+	};
+
+	/**
+	 * Calculates the bounds of the instanced static mesh from the mass fragments
+	 * @param TransformFragment Transform to use as a root of the ISM
+	 * @param RenderISMFragment List of instances transform of the ISM
+	 * @param BoundsType Bounds type to calculate
+	 * @return 
+	 */
+	UE_API FBoxSphereBounds CalculateInstancedStaticMeshBounds(const FTransformFragment& TransformFragment, const FMassRenderISMFragment& RenderISMFragment, const EBoundsType BoundsType);
+
+#if WITH_EDITOR
+	/**
+	 * Initialize the scene proxy desc from the editor mesh fragments
+	 * @param EditorMeshFragment fragment describing the editor mesh parameters
+	 * @param PrimitiveSceneProxyDesc scene proxy desc to initialize
+	 */
+	void InitializePrimitiveSceneProxyDescFromEditorFragment(const FMassEditorVisualizationMeshFragment& EditorMeshFragment, FPrimitiveSceneProxyDesc& PrimitiveSceneProxyDesc);
+#endif // WITH_EDITOR
+
+	/**
+	 * Initialize the static mesh scene proxy desc from mass fragments
+	 * @param StaticMeshFragment fragment describing the static mesh
+	 * @param MeshFragment fragment describing the mesh parameters
+	 * @param TransformFragment fragment specifying the transform
+	 * @param StaticMeshSceneProxyDesc scene proxy desc to initialize
+	 */
+	void InitializeStaticMeshSceneProxyDescFromFragment(const FMassStaticMeshFragment& StaticMeshFragment, const FMassVisualizationMeshFragment& MeshFragment, const FTransformFragment& TransformFragment, FStaticMeshSceneProxyDesc& StaticMeshSceneProxyDesc);
+
+	/**
+	 * Initialize the instanced static mesh scene proxy desc from mass fragments
+	 * @param StaticMeshFragment fragment describing the static mesh
+	 * @param MeshFragment fragment describing the mesh parameters
+	 * @param TransformFragment fragment specifying the transform
+	 * @param ISMSceneProxyDesc scene proxy desc to initialize
+	 */
+	void InitializeInstanceStaticMeshSceneProxyDescFromFragment(const FMassStaticMeshFragment& StaticMeshFragment, const FMassVisualizationMeshFragment& MeshFragment, const FTransformFragment& TransformFragment, FInstancedStaticMeshSceneProxyDesc& ISMSceneProxyDesc);
+
+
+}
+
+#undef UE_API
\ No newline at end of file
diff --git a/Engine/Source/Runtime/Mass/MassEngine/Public/Mesh/MassEngineRenderISMProcessors.h b/Engine/Source/Runtime/Mass/MassEngine/Public/Mesh/MassEngineRenderISMProcessors.h
new file mode 100644
--- /dev/null
+++ b/Engine/Source/Runtime/Mass/MassEngine/Public/Mesh/MassEngineRenderISMProcessors.h
@@ -0,0 +1,76 @@
+// Copyright Epic Games, Inc. All Rights Reserved.
+
+#pragma once
+
+#include "Mesh/MassEngineRenderStaticMeshProcessors.h"
+#include "MassObserverProcessor.h"
+#include "MassSignalProcessorBase.h"
+#include "MassArchetypeTypes.h"
+
+#include "MassEngineRenderISMProcessors.generated.h"
+
+struct FMassEntityManager;
+struct FMassExecutionContext;
+
+/**
+ * Processor responsible for creating instanced static mesh entity and setting up its render state for the renderer
+ */
+UCLASS()
+class UMassInstantiateStaticMeshAndCreateISMRenderStateProcessor : public UMassBaseStaticMeshSetupRenderStateProcessor
+{
+	GENERATED_BODY()
+
+protected:
+	virtual void ConfigureQueries(const TSharedRef<FMassEntityManager>& EntityManager) override;
+	virtual void Execute(FMassEntityManager& EntityManager, FMassExecutionContext& ExecutionContext) override;
+
+	FMassElementBitSet ISMEntityComposition;
+	FMassArchetypeHandle ISMEntityArchetype;
+	FMassElementBitSet ISMEntityCompositionWithMaterialOverrides;
+	FMassArchetypeHandle ISMEntityArchetypeWithMaterialOverrides;
+
+#if WITH_EDITOR
+	FMassElementBitSet ISMEntityCompositionWithEditorMesh;
+	FMassArchetypeHandle ISMEntityArchetypeWithEditorMesh;
+	FMassElementBitSet ISMEntityCompositionWithMaterialOverridesAndEditorMesh;
+	FMassArchetypeHandle ISMEntityArchetypeWithMaterialOverridesAndEditorMesh;
+#endif // WITH_EDITOR
+};
+
+/**
+ * Processor responsible for notifying the linked instanced static mesh entity to update the render proxy
+ * This processor listen to the following signals on the instances entities
+ * - UE::Mass::Signals::TransformChanged
+ * - UE::Mass::Signals::SelectionChanged
+ */
+UCLASS()
+class UMassDirtyRenderISMInstanceProcessor : public UMassSignalProcessorBase
+{
+	GENERATED_BODY()
+
+public:
+	explicit UMassDirtyRenderISMInstanceProcessor(const FObjectInitializer& ObjectInitializer = FObjectInitializer::Get());
+
+protected:
+	virtual void InitializeInternal(UObject& Owner, const TSharedRef<FMassEntityManager>& EntityManager) override;
+	virtual void ConfigureQueries(const TSharedRef<FMassEntityManager>& EntityManager) override;
+	virtual void SignalEntities(FMassEntityManager& EntityManager, FMassExecutionContext& Context, FMassSignalNameLookup& EntitySignals) override;
+};
+
+/**
+ * Processor responsible for removing instances and destroying the render state of the instanced static mesh entity
+ */
+UCLASS()
+class UMassRemoveInstanceAndDestroyISMRenderStateProcessor : public UMassObserverProcessor
+{
+	GENERATED_BODY()
+
+public:
+	UMassRemoveInstanceAndDestroyISMRenderStateProcessor();
+
+protected:
+	virtual void ConfigureQueries(const TSharedRef<FMassEntityManager>& EntityManager) override;
+	virtual void Execute(FMassEntityManager& EntityManager, FMassExecutionContext& Context) override;
+
+	FMassEntityQuery EntityQuery;
+};
\ No newline at end of file
diff --git a/Engine/Source/Runtime/Mass/MassEngine/Public/Mesh/MassEngineRenderStateProcessors.h b/Engine/Source/Runtime/Mass/MassEngine/Public/Mesh/MassEngineRenderStateProcessors.h
new file mode 100644
--- /dev/null
+++ b/Engine/Source/Runtime/Mass/MassEngine/Public/Mesh/MassEngineRenderStateProcessors.h
@@ -0,0 +1,71 @@
+// Copyright Epic Games, Inc. All Rights Reserved.
+
+#pragma once
+
+#include "MassObserverProcessor.h"
+#include "MassSignalProcessorBase.h"
+
+#include "MassEngineRenderStateProcessors.generated.h"
+
+struct FMassEntityManager;
+struct FMassExecutionContext;
+
+/**
+ * Processor responsible for creating the renderer states
+ */
+UCLASS()
+class UMassCreateRenderStateProcessor : public UMassProcessor
+{
+	GENERATED_BODY()
+
+public:
+	UMassCreateRenderStateProcessor(const FObjectInitializer& ObjectInitializer = FObjectInitializer::Get());
+
+protected:
+	virtual void ConfigureQueries(const TSharedRef<FMassEntityManager>& EntityManager) override;
+	virtual void Execute(FMassEntityManager& EntityManager, FMassExecutionContext& ExecutionContext) override;
+
+	FMassEntityQuery EntityQuery;
+};
+
+/**
+ * Processor responsible for updating or recreating the render state when dirtied
+ * Listens to the following signals:
+ * - UE::Mass::Signals::TransformChanged
+ * - UE::Mass::Signals::MeshVisualPropertyChanged
+ * - UE::Mass::Signals::MeshChanged
+ * - UE::Mass::Signals::SelectionChanged
+ * - UE::Mass::Signals::LevelEditingStateChanged
+ * - UE::Mass::Signals::RenderStateDirty
+ */
+UCLASS()
+class UMassRenderStateDirtyUpdateProcessor : public UMassSignalProcessorBase
+{
+	GENERATED_BODY()
+
+public:
+	UMassRenderStateDirtyUpdateProcessor(const FObjectInitializer& ObjectInitializer = FObjectInitializer::Get());
+
+protected:
+	virtual void InitializeInternal(UObject& Owner, const TSharedRef<FMassEntityManager>& EntityManager) override;
+	virtual void ConfigureQueries(const TSharedRef<FMassEntityManager>& EntityManager) override;
+	virtual void SignalEntities(FMassEntityManager& EntityManager, FMassExecutionContext& Context, FMassSignalNameLookup& EntitySignals) override;
+};
+
+/**
+ * Processor responsible for destroying the render states
+ */
+UCLASS()
+class UMassDestroyRenderStateProcessor : public UMassObserverProcessor
+{
+	GENERATED_BODY()
+
+public:
+	UMassDestroyRenderStateProcessor();
+
+protected:
+	virtual void ConfigureQueries(const TSharedRef<FMassEntityManager>& EntityManager) override;
+	virtual void Execute(FMassEntityManager& EntityManager, FMassExecutionContext& Context) override;
+
+	FMassEntityQuery EntityQuery;
+};
\ No newline at end of file
diff --git a/Engine/Source/Runtime/Mass/MassEngine/Public/Mesh/MassEngineRenderStaticMeshProcessors.h b/Engine/Source/Runtime/Mass/MassEngine/Public/Mesh/MassEngineRenderStaticMeshProcessors.h
new file mode 100644
--- /dev/null
+++ b/Engine/Source/Runtime/Mass/MassEngine/Public/Mesh/MassEngineRenderStaticMeshProcessors.h
@@ -0,0 +1,64 @@
+// Copyright Epic Games, Inc. All Rights Reserved.
+
+#pragma once
+
+#include "MassObserverProcessor.h"
+#include "MassSignalProcessorBase.h"
+
+#include "MassEngineRenderStaticMeshProcessors.generated.h"
+
+struct FMassEntityManager;
+struct FMassExecutionContext;
+
+/**
+ * Base processor responsible for setting up the static mesh render state for the renderer
+ */
+UCLASS(Abstract)
+class UMassBaseStaticMeshSetupRenderStateProcessor : public UMassProcessor
+{
+	GENERATED_BODY()
+
+public:
+	UMassBaseStaticMeshSetupRenderStateProcessor(const FObjectInitializer& ObjectInitializer = FObjectInitializer::Get());
+
+protected:
+	virtual void ConfigureQueries(const TSharedRef<FMassEntityManager>& EntityManager) override;
+
+	FMassEntityQuery EntityQuery;
+};
+
+/**
+ * Processor responsible for setting up the static mesh render state for the renderer
+ */
+UCLASS()
+class UMassStaticMeshSetupRenderStateProcessor : public UMassBaseStaticMeshSetupRenderStateProcessor
+{
+	GENERATED_BODY()
+
+protected:
+	virtual void ConfigureQueries(const TSharedRef<FMassEntityManager>& EntityManager) override;
+	virtual void Execute(FMassEntityManager& EntityManager, FMassExecutionContext& ExecutionContext) override;
+};
+
+/**
+ * Processor responsible for updating the static mesh render state when dirtied
+ * Listening the following signals on the entity
+ * - UE::Mass::Signals::TransformChanged
+ * - UE::Mass::Signals::MeshVisualPropertyChanged
+ * - UE::Mass::Signals::MeshChanged
+ * - UE::Mass::Signals::SelectionChanged
+ * - UE::Mass::Signals::LevelEditingStateChanged
+ */
+UCLASS()
+class UMassStaticMeshUpdateRenderStateProcessor : public UMassSignalProcessorBase
+{
+	GENERATED_BODY()
+
+public:
+	UMassStaticMeshUpdateRenderStateProcessor(const FObjectInitializer& ObjectInitializer = FObjectInitializer::Get());
+
+protected:
+	virtual void InitializeInternal(UObject& Owner, const TSharedRef<FMassEntityManager>& EntityManager) override;
+	virtual void ConfigureQueries(const TSharedRef<FMassEntityManager>& EntityManager) override;
+	virtual void SignalEntities(FMassEntityManager& EntityManager, FMassExecutionContext& Context, FMassSignalNameLookup& EntitySignals) override;
+};
