Fix build break

#rnx

#Synced-CL 51263739

diff --git a/Engine/Source/Runtime/Mass/MassEngine/Private/Mesh/MassEngineMeshFragments.cpp b/Engine/Source/Runtime/Mass/MassEngine/Private/Mesh/MassEngineMeshFragments.cpp
--- a/Engine/Source/Runtime/Mass/MassEngine/Private/Mesh/MassEngineMeshFragments.cpp
+++ b/Engine/Source/Runtime/Mass/MassEngine/Private/Mesh/MassEngineMeshFragments.cpp
@@ -92,7 +92,9 @@ FMassRenderISMFragment::~FMassRenderISMFragment()
 
 FMassRenderISMFragment::FMassRenderISMFragment(FMassRenderISMFragment&& Other)
 : FMassFragment(MoveTemp(Other))
+#if WITH_EDITORONLY_DATA
 , PerInstanceHitProxy(MoveTemp(Other.PerInstanceHitProxy))
+#endif // WITH_EDITORONLY_DATA
 , PerInstanceSMData(MoveTemp(Other.PerInstanceSMData))
 , PerInstanceSMCustomData(MoveTemp(Other.PerInstanceSMCustomData))
 , AdditionalRandomSeeds(MoveTemp(Other.AdditionalRandomSeeds))
