Fix build break

#rnx

#Synced-CL 51263739

diff --git a/Engine/Source/Runtime/Mass/MassEngine/Private/Mesh/MassISMRenderStateHelper.cpp b/Engine/Source/Runtime/Mass/MassEngine/Private/Mesh/MassISMRenderStateHelper.cpp
--- a/Engine/Source/Runtime/Mass/MassEngine/Private/Mesh/MassISMRenderStateHelper.cpp
+++ b/Engine/Source/Runtime/Mass/MassEngine/Private/Mesh/MassISMRenderStateHelper.cpp
@@ -137,7 +137,8 @@ FPrimitiveSceneProxy* FMassISMRenderStateHelper::CreateStaticMeshSceneProxy(cons
 #if !WITH_EDITOR && UE_WITH_PSO_PRECACHING
 void FMassISMRenderStateHelper::CollectPSOPrecacheData(const FPSOPrecacheParams& BasePrecachePSOParams, FMaterialInterfacePSOPrecacheParamsList& OutParams)
 {
-	return FInstancedStaticMeshComponentHelper::CollectPSOPrecacheData(*this, BasePrecachePSOParams, OutParams);
+	// @TODO Implement this
+	//FInstancedStaticMeshComponentHelper::CollectPSOPrecacheData(*this, BasePrecachePSOParams, OutParams);
 }
 #endif
 
