Fix for an other place of the same build break.

#rnx

#Synced-CL 51263739

diff --git a/Engine/Source/Runtime/Mass/MassEngine/Private/Mesh/MassStaticMeshRenderStateHelper.cpp b/Engine/Source/Runtime/Mass/MassEngine/Private/Mesh/MassStaticMeshRenderStateHelper.cpp
--- a/Engine/Source/Runtime/Mass/MassEngine/Private/Mesh/MassStaticMeshRenderStateHelper.cpp
+++ b/Engine/Source/Runtime/Mass/MassEngine/Private/Mesh/MassStaticMeshRenderStateHelper.cpp
@@ -315,7 +315,8 @@ bool FMassBaseStaticMeshRenderStateHelper::IsDisplayNaniteFallbackMesh() const
 #if !WITH_EDITOR && UE_WITH_PSO_PRECACHING
 void FMassBaseStaticMeshRenderStateHelper::CollectPSOPrecacheData(const FPSOPrecacheParams& BasePrecachePSOParams, FMaterialInterfacePSOPrecacheParamsList& OutParams)
 {
-	FStaticMeshComponentHelper::CollectPSOPrecacheData(*this, BasePrecachePSOParams, OutParams);
+	// @TODO Implement this
+	//FStaticMeshComponentHelper::CollectPSOPrecacheData(*this, BasePrecachePSOParams, OutParams);
 }
 #endif
 
