Fix static analysis

#Synced-CL 51334569

diff --git a/Engine/Source/Runtime/Mass/MassEngine/Private/Mesh/MassEngineRenderISMProcessors.cpp b/Engine/Source/Runtime/Mass/MassEngine/Private/Mesh/MassEngineRenderISMProcessors.cpp
--- a/Engine/Source/Runtime/Mass/MassEngine/Private/Mesh/MassEngineRenderISMProcessors.cpp
+++ b/Engine/Source/Runtime/Mass/MassEngine/Private/Mesh/MassEngineRenderISMProcessors.cpp
@@ -99,12 +99,7 @@ void UMassInstantiateStaticMeshAndCreateISMRenderStateProcessor::Execute(FMassEn
 
 		// Build the source shared fragment hash
 		FHashBuilder SourceSharedFragmentsHash;
-		auto AddFragmentToHash = [&SourceSharedFragmentsHash](const void* FragmentPtr)
-			{
-				SourceSharedFragmentsHash.AppendRaw(&FragmentPtr, sizeof(&FragmentPtr));
-			};
-		AddFragmentToHash(&StaticMeshFragment);
-		AddFragmentToHash(&MeshFragment);
+		SourceSharedFragmentsHash << &StaticMeshFragment << &MeshFragment;
 
 		// Find if the ISM entity already exist, if not add it
 		FISMEntityInfo* ISMEntity;
