Replaces manual raw pointer ownership in FMassRenderStaticMeshFragment, FMassRenderISMFragment, and  FMassRenderStateFragment with TSharedPtr/TSharedRef

#rnx
#Synced-CL 51455761

diff --git a/Engine/Source/Runtime/Mass/MassEngine/Private/Mesh/MassEngineMeshFragments.cpp b/Engine/Source/Runtime/Mass/MassEngine/Private/Mesh/MassEngineMeshFragments.cpp
--- a/Engine/Source/Runtime/Mass/MassEngine/Private/Mesh/MassEngineMeshFragments.cpp
+++ b/Engine/Source/Runtime/Mass/MassEngine/Private/Mesh/MassEngineMeshFragments.cpp
@@ -17,34 +17,15 @@ FMassStaticMeshFragment::FMassStaticMeshFragment(TNotNull<const UStaticMesh*> In
 }
 
 FMassRenderStaticMeshFragment::FMassRenderStaticMeshFragment()
-	: StaticMeshSceneProxyDesc(new FStaticMeshSceneProxyDesc())
-	, bOwnProxyDesc(true)
+	: StaticMeshSceneProxyDesc(MakeShared<FStaticMeshSceneProxyDesc>())
 {
 }
 
-FMassRenderStaticMeshFragment::FMassRenderStaticMeshFragment(TNotNull<FStaticMeshSceneProxyDesc*> InStaticMeshSceneProxyDesc)
+FMassRenderStaticMeshFragment::FMassRenderStaticMeshFragment(const TSharedRef<FStaticMeshSceneProxyDesc>& InStaticMeshSceneProxyDesc)
 	: StaticMeshSceneProxyDesc(InStaticMeshSceneProxyDesc)
-	, bOwnProxyDesc(false)
 {
 }
 
-FMassRenderStaticMeshFragment::~FMassRenderStaticMeshFragment()
-{
-	if (bOwnProxyDesc)
-	{
-		delete StaticMeshSceneProxyDesc;
-	}
-}
-
-FMassRenderStaticMeshFragment::FMassRenderStaticMeshFragment(FMassRenderStaticMeshFragment&& Other)
-	: FMassFragment(MoveTemp(Other))
-	, StaticMeshSceneProxyDesc(Other.StaticMeshSceneProxyDesc)
-	, bOwnProxyDesc(Other.bOwnProxyDesc)
-{
-	Other.bOwnProxyDesc = false;
-	Other.StaticMeshSceneProxyDesc = nullptr;
-}
-
 UStaticMesh* FMassRenderStaticMeshFragment::GetStaticMesh()
 {
 	return StaticMeshSceneProxyDesc ? StaticMeshSceneProxyDesc->StaticMesh : nullptr;
@@ -71,39 +52,13 @@ EMouseCursor::Type HTypeElementHandleHitProxy::GetMouseCursor()
 // FMassRenderISMFragment
 //----------------------------------------------------------------------//
 FMassRenderISMFragment::FMassRenderISMFragment()
-	: InstancedStaticMeshSceneProxyDesc(new FInstancedStaticMeshSceneProxyDesc())
-	, bOwnProxyDesc(true)
+	: InstancedStaticMeshSceneProxyDesc(MakeShared<FInstancedStaticMeshSceneProxyDesc>())
 {
 }
 
-FMassRenderISMFragment::FMassRenderISMFragment(TNotNull<FInstancedStaticMeshSceneProxyDesc*> InInstancedStaticMeshSceneProxyDesc)
+FMassRenderISMFragment::FMassRenderISMFragment(const TSharedRef<FInstancedStaticMeshSceneProxyDesc>& InInstancedStaticMeshSceneProxyDesc)
 	: InstancedStaticMeshSceneProxyDesc(InInstancedStaticMeshSceneProxyDesc)
-	, bOwnProxyDesc(false)
-{
-}
-
-FMassRenderISMFragment::~FMassRenderISMFragment()
-{
-	if (bOwnProxyDesc)
-	{
-		delete InstancedStaticMeshSceneProxyDesc;
-	}
-}
-
-FMassRenderISMFragment::FMassRenderISMFragment(FMassRenderISMFragment&& Other)
-: FMassFragment(MoveTemp(Other))
-#if WITH_EDITORONLY_DATA
-, PerInstanceHitProxy(MoveTemp(Other.PerInstanceHitProxy))
-#endif // WITH_EDITORONLY_DATA
-, PerInstanceSMData(MoveTemp(Other.PerInstanceSMData))
-, PerInstanceSMCustomData(MoveTemp(Other.PerInstanceSMCustomData))
-, AdditionalRandomSeeds(MoveTemp(Other.AdditionalRandomSeeds))
-, InstancedStaticMeshSceneProxyDesc(Other.InstancedStaticMeshSceneProxyDesc)
-, InstancingRandomSeed(Other.InstancingRandomSeed)
-, bOwnProxyDesc(Other.bOwnProxyDesc)
 {
-	Other.bOwnProxyDesc = false;
-	Other.InstancedStaticMeshSceneProxyDesc = nullptr;
 }
 
 UStaticMesh* FMassRenderISMFragment::GetStaticMesh()
diff --git a/Engine/Source/Runtime/Mass/MassEngine/Private/Mesh/MassEngineMeshUtils.cpp b/Engine/Source/Runtime/Mass/MassEngine/Private/Mesh/MassEngineMeshUtils.cpp
--- a/Engine/Source/Runtime/Mass/MassEngine/Private/Mesh/MassEngineMeshUtils.cpp
+++ b/Engine/Source/Runtime/Mass/MassEngine/Private/Mesh/MassEngineMeshUtils.cpp
@@ -25,6 +25,8 @@ namespace UE::MassEngine::Mesh
 		InitializePrimitiveFragmentsFromComponent(Component, TransformFragment, RenderPrimitiveFragment);
 
 		const UStaticMeshComponent* StaticMeshComponent = CastChecked<const UStaticMeshComponent>(Component);
+
+		checkf(RenderStaticMeshFragment.StaticMeshSceneProxyDesc.IsValid(), TEXT("Expecting a valid scene proxy desc"));
 		RenderStaticMeshFragment.StaticMeshSceneProxyDesc->InitializeFromStaticMeshComponent(StaticMeshComponent);
 
 		UStaticMesh* StaticMesh = RenderStaticMeshFragment.StaticMeshSceneProxyDesc->StaticMesh;
@@ -68,13 +70,9 @@ namespace UE::MassEngine::Mesh
 	{
 		InitializePrimitiveFragmentsFromComponent(Component, TransformFragment, RenderPrimitiveFragment);
 
-		if (!RenderISMFragment.InstancedStaticMeshSceneProxyDesc)
-		{
-			RenderISMFragment.InstancedStaticMeshSceneProxyDesc = new FInstancedStaticMeshSceneProxyDesc();
-			RenderISMFragment.bOwnProxyDesc = true;
-		}
-
 		const UInstancedStaticMeshComponent* ISMComponent = CastChecked<const UInstancedStaticMeshComponent>(Component);
+
+		checkf(RenderISMFragment.InstancedStaticMeshSceneProxyDesc.IsValid(), TEXT("Expecting a valid scene proxy desc"));
 		RenderISMFragment.InstancedStaticMeshSceneProxyDesc->InitializeFromStaticMeshComponent(ISMComponent);
 
 		RenderISMFragment.InstancedStaticMeshSceneProxyDesc->bCollisionEnabled = !!RenderISMFragment.InstancedStaticMeshSceneProxyDesc->bCollisionEnabled && !ISMComponent->bDisableCollision;
diff --git a/Engine/Source/Runtime/Mass/MassEngine/Private/Mesh/MassRenderStateHelper.cpp b/Engine/Source/Runtime/Mass/MassEngine/Private/Mesh/MassRenderStateHelper.cpp
--- a/Engine/Source/Runtime/Mass/MassEngine/Private/Mesh/MassRenderStateHelper.cpp
+++ b/Engine/Source/Runtime/Mass/MassEngine/Private/Mesh/MassRenderStateHelper.cpp
@@ -8,28 +8,14 @@
 #include "MassSignalSubsystem.h"
 #include "Mesh/MassEngineMeshFragments.h"
 
-
-FMassRenderStateFragment::~FMassRenderStateFragment()
-{
-	DestroyRenderStateHelper();
-}
-
 // ----------------------------------------------------------------------//
 // FMassRenderStateFragment
 //----------------------------------------------------------------------//
 void FMassRenderStateFragment::DestroyRenderStateHelper()
 {
-	delete RenderStateHelper;
 	RenderStateHelper = nullptr;
 }
 
-FMassRenderStateFragment::FMassRenderStateFragment(FMassRenderStateFragment&& Other)
-	: FMassFragment(std::move(Other))
-    , RenderStateHelper(Other.RenderStateHelper)
-{
-	Other.RenderStateHelper = nullptr;
-}
-
 // ----------------------------------------------------------------------//
 // FMassRenderStateHelper
 //----------------------------------------------------------------------//
diff --git a/Engine/Source/Runtime/Mass/MassEngine/Private/Mesh/MassRenderStateHelper.h b/Engine/Source/Runtime/Mass/MassEngine/Private/Mesh/MassRenderStateHelper.h
--- a/Engine/Source/Runtime/Mass/MassEngine/Private/Mesh/MassRenderStateHelper.h
+++ b/Engine/Source/Runtime/Mass/MassEngine/Private/Mesh/MassRenderStateHelper.h
@@ -33,13 +33,7 @@ struct FMassRenderStateFragment : public FMassFragment
 	GENERATED_BODY()
 
 	FMassRenderStateFragment() = default;
-	~FMassRenderStateFragment();
 	void DestroyRenderStateHelper();
-	FMassRenderStateFragment(const FMassRenderStateFragment& Other) = delete;
-	FMassRenderStateFragment(FMassRenderStateFragment&& Other);
-
-	FMassRenderStateFragment& operator=(const FMassRenderStateFragment& Other) = delete;
-	FMassRenderStateFragment& operator=(FMassRenderStateFragment&& Other) = delete;
 
 	template<typename T = FMassRenderStateHelper>
 	const T& GetRenderStateHelper() const requires TIsDerivedFrom<T, FMassRenderStateHelper>::Value
@@ -61,23 +55,14 @@ struct FMassRenderStateFragment : public FMassFragment
 	T* CreateRenderStateHelper(TArgs&&... InArgs) requires TIsDerivedFrom<T, FMassRenderStateHelper>::Value
 	{
 		checkf(!RenderStateHelper, TEXT("Render State Helper already exist"));
-		T* CreatedRenderState = new T(Forward<TArgs>(InArgs)...);
+		TSharedPtr<T> CreatedRenderState = MakeShared<T>(Forward<TArgs>(InArgs)...);
 		RenderStateHelper = CreatedRenderState;
-		return CreatedRenderState;
+		return CreatedRenderState.Get();
 	}
 
 private:
-	FMassRenderStateHelper* RenderStateHelper = nullptr;
-
-};
+	TSharedPtr<FMassRenderStateHelper> RenderStateHelper = nullptr;
 
-template<>
-struct TStructOpsTypeTraits<FMassRenderStateFragment> : TStructOpsTypeTraitsBase2<FMassRenderStateFragment>
-{
-	enum
-	{
-		WithCopy = false,
-	};
 };
 
 template<>
diff --git a/Engine/Source/Runtime/Mass/MassEngine/Public/Mesh/MassEngineMeshFragments.h b/Engine/Source/Runtime/Mass/MassEngine/Public/Mesh/MassEngineMeshFragments.h
--- a/Engine/Source/Runtime/Mass/MassEngine/Public/Mesh/MassEngineMeshFragments.h
+++ b/Engine/Source/Runtime/Mass/MassEngine/Public/Mesh/MassEngineMeshFragments.h
@@ -154,28 +154,12 @@ struct FMassRenderStaticMeshFragment : public FMassFragment
 	GENERATED_BODY()
 
 	UE_API FMassRenderStaticMeshFragment();
-	UE_API explicit FMassRenderStaticMeshFragment(TNotNull<FStaticMeshSceneProxyDesc*> InStaticMeshSceneProxyDesc);
-	UE_API ~FMassRenderStaticMeshFragment();
-	FMassRenderStaticMeshFragment(const FMassRenderStaticMeshFragment& Other) = delete;
-	FMassRenderStaticMeshFragment(FMassRenderStaticMeshFragment&& Other);
-
-	FMassRenderStaticMeshFragment& operator=(const FMassRenderStaticMeshFragment& Other) = delete;
-	FMassRenderStaticMeshFragment& operator=(FMassRenderStaticMeshFragment&& Other) = delete;
+	UE_API explicit FMassRenderStaticMeshFragment(const TSharedRef<FStaticMeshSceneProxyDesc>& InStaticMeshSceneProxyDesc);
 
 	UStaticMesh* GetStaticMesh();
 	const UStaticMesh* GetStaticMesh() const;
 
-	FStaticMeshSceneProxyDesc* StaticMeshSceneProxyDesc = nullptr;
-	bool bOwnProxyDesc : 1 = false;
-};
-
-template<>
-struct TStructOpsTypeTraits<FMassRenderStaticMeshFragment> : TStructOpsTypeTraitsBase2<FMassRenderStaticMeshFragment>
-{
-	enum
-	{
-		WithCopy = false,
-	};
+	TSharedPtr<FStaticMeshSceneProxyDesc> StaticMeshSceneProxyDesc = nullptr;
 };
 
 template<>
@@ -236,16 +220,7 @@ struct FMassRenderISMFragment : public FMassFragment
 	GENERATED_BODY()
 
 	FMassRenderISMFragment();
-	explicit FMassRenderISMFragment(TNotNull<FInstancedStaticMeshSceneProxyDesc*> InInstancedStaticMeshSceneProxyDesc);
-	~FMassRenderISMFragment();
-
-	FMassRenderISMFragment(const FMassRenderISMFragment& Other) = delete;
-	FMassRenderISMFragment(FMassRenderISMFragment&& Other);
-
-	FMassRenderISMFragment& operator=(const FMassRenderISMFragment& Other) = delete;
-	FMassRenderISMFragment& operator=(FMassRenderISMFragment&& Other) = delete;
-
-	//FMassRenderISMFragment(FMassRenderISMFragment& Other);
+	explicit FMassRenderISMFragment(const TSharedRef<FInstancedStaticMeshSceneProxyDesc>& InInstancedStaticMeshSceneProxyDesc);
 
 	UStaticMesh* GetStaticMesh();
 	const UStaticMesh* GetStaticMesh() const;
@@ -257,18 +232,8 @@ struct FMassRenderISMFragment : public FMassFragment
 	TArray<float> PerInstanceSMCustomData;
 	TArray<FInstancedStaticMeshRandomSeed> AdditionalRandomSeeds;
 
-	FInstancedStaticMeshSceneProxyDesc* InstancedStaticMeshSceneProxyDesc = nullptr;
+	TSharedPtr<FInstancedStaticMeshSceneProxyDesc> InstancedStaticMeshSceneProxyDesc = nullptr;
 	int32 InstancingRandomSeed = 0;
-	uint8 bOwnProxyDesc : 1 = false;
-};
-
-template<>
-struct TStructOpsTypeTraits<FMassRenderISMFragment> : TStructOpsTypeTraitsBase2<FMassRenderISMFragment>
-{
-	enum
-	{
-		WithCopy = false,
-	};
 };
 
 template<>
