[Mass] Added FName-based constructor to FMassArchetypeCreationParams and pass debug names to CreateArchetype call sites for improved archetype identification

#Synced-CL 51846703

#Synced-CL 51846703

diff --git a/Engine/Plugins/AI/MassAI/Source/MassAIBehavior/Private/MassLookAtSubsystem.cpp b/Engine/Plugins/AI/MassAI/Source/MassAIBehavior/Private/MassLookAtSubsystem.cpp
--- a/Engine/Plugins/AI/MassAI/Source/MassAIBehavior/Private/MassLookAtSubsystem.cpp
+++ b/Engine/Plugins/AI/MassAI/Source/MassAIBehavior/Private/MassLookAtSubsystem.cpp
@@ -424,12 +424,12 @@ void UMassLookAtSubsystem::Initialize(FSubsystemCollectionBase& Collection)
 
 	// Create Mass archetype for entities representing requests
 	const FMassArchetypeCompositionDescriptor RequestComposition{ FMassFragmentBitSet(FMassLookAtRequestFragment::StaticStruct()) };
-	RequestArchetype = EntityManager.CreateArchetype(RequestComposition);
+	RequestArchetype = EntityManager.CreateArchetype(RequestComposition, FMassArchetypeCreationParams{"LookAtRequest"});
 
 	// Create Mass archetype for entities representing targets
 	TArray<const UScriptStruct*, TInlineAllocator<2>> TargetFragmentTypes{FMassLookAtTargetFragment::StaticStruct(), FTransformFragment::StaticStruct()};
 	const FMassArchetypeCompositionDescriptor TargetComposition{ FMassFragmentBitSet(TargetFragmentTypes) };
-	TargetArchetype = EntityManager.CreateArchetype(TargetComposition);
+	TargetArchetype = EntityManager.CreateArchetype(TargetComposition, FMassArchetypeCreationParams{"LookAtTarget"});
 
 	OverrideSubsystemTraits<UMassLookAtSubsystem>(Collection);
 }
diff --git a/Engine/Source/Runtime/Mass/MassEngine/Private/Mesh/MassEngineRenderISMProcessors.cpp b/Engine/Source/Runtime/Mass/MassEngine/Private/Mesh/MassEngineRenderISMProcessors.cpp
--- a/Engine/Source/Runtime/Mass/MassEngine/Private/Mesh/MassEngineRenderISMProcessors.cpp
+++ b/Engine/Source/Runtime/Mass/MassEngine/Private/Mesh/MassEngineRenderISMProcessors.cpp
@@ -32,23 +32,23 @@ void UMassInstantiateStaticMeshAndCreateISMRenderStateProcessor::ConfigureQuerie
 	ISMEntityComposition.Add<FMassRenderStateFragment>();
 	ISMEntityComposition.Add<FMassRenderPrimitiveFragment>();
 	ISMEntityComposition.Add<FMassRenderISMFragment>();
-	ISMEntityArchetype = EntityManager->CreateArchetype(ISMEntityComposition);
+	ISMEntityArchetype = EntityManager->CreateArchetype(ISMEntityComposition, FMassArchetypeCreationParams{"RenderISM"});
 	ProcessorEntityCreationRequirements.AddCreatedArchetype(ISMEntityComposition);
 
 	ISMEntityCompositionWithMaterialOverrides = ISMEntityComposition;
 	ISMEntityCompositionWithMaterialOverrides.Add<FMassOverrideMaterialsFragment>();
-	ISMEntityArchetypeWithMaterialOverrides = EntityManager->CreateArchetype(ISMEntityCompositionWithMaterialOverrides);
+	ISMEntityArchetypeWithMaterialOverrides = EntityManager->CreateArchetype(ISMEntityCompositionWithMaterialOverrides, FMassArchetypeCreationParams{"RenderISMwMaterialOverrides"});
 	ProcessorEntityCreationRequirements.AddCreatedArchetype(ISMEntityCompositionWithMaterialOverrides);
 
 #if WITH_EDITOR
 	ISMEntityCompositionWithEditorMesh = ISMEntityComposition;
 	ISMEntityCompositionWithEditorMesh.Add<FMassEditorVisualizationMeshFragment>();
-	ISMEntityArchetypeWithEditorMesh = EntityManager->CreateArchetype(ISMEntityCompositionWithEditorMesh);
+	ISMEntityArchetypeWithEditorMesh = EntityManager->CreateArchetype(ISMEntityCompositionWithEditorMesh, FMassArchetypeCreationParams{"RenderISMwEditorMesh"});
 	ProcessorEntityCreationRequirements.AddCreatedArchetype(ISMEntityCompositionWithEditorMesh);
 
 	ISMEntityCompositionWithMaterialOverridesAndEditorMesh = ISMEntityCompositionWithMaterialOverrides;
 	ISMEntityCompositionWithMaterialOverridesAndEditorMesh.Add<FMassEditorVisualizationMeshFragment>();
-	ISMEntityArchetypeWithMaterialOverridesAndEditorMesh = EntityManager->CreateArchetype(ISMEntityCompositionWithMaterialOverridesAndEditorMesh);
+	ISMEntityArchetypeWithMaterialOverridesAndEditorMesh = EntityManager->CreateArchetype(ISMEntityCompositionWithMaterialOverridesAndEditorMesh, FMassArchetypeCreationParams{"RenderISMwMaterialOverridesAndEditorMesh"});
 	ProcessorEntityCreationRequirements.AddCreatedArchetype(ISMEntityCompositionWithMaterialOverridesAndEditorMesh);
 #endif // WITH_EDITOR
 
diff --git a/Engine/Source/Runtime/MassEntity/Public/MassEntityTypes.h b/Engine/Source/Runtime/MassEntity/Public/MassEntityTypes.h
--- a/Engine/Source/Runtime/MassEntity/Public/MassEntityTypes.h
+++ b/Engine/Source/Runtime/MassEntity/Public/MassEntityTypes.h
@@ -942,6 +942,11 @@ namespace UE::Mass
 struct FMassArchetypeCreationParams
 {
 	FMassArchetypeCreationParams() = default;
+	explicit FMassArchetypeCreationParams(const FName DebugName)
+		: DebugName(DebugName)
+	{
+	}
+
 	explicit FMassArchetypeCreationParams(const struct FMassArchetypeData& Archetype);
 
 	/** Created archetype will have chunks of this size. 0 denotes "use default" (see UE::Mass::ChunkSize) */
