Removed temporary log

#Synced-CL 51945655

diff --git a/Engine/Source/Runtime/Mass/MassEngine/Private/Mesh/MassISMRenderStateHelper.cpp b/Engine/Source/Runtime/Mass/MassEngine/Private/Mesh/MassISMRenderStateHelper.cpp
--- a/Engine/Source/Runtime/Mass/MassEngine/Private/Mesh/MassISMRenderStateHelper.cpp
+++ b/Engine/Source/Runtime/Mass/MassEngine/Private/Mesh/MassISMRenderStateHelper.cpp
@@ -94,7 +94,6 @@ TSharedPtr<FInstanceDataSceneProxy, ESPMode::ThreadSafe>& FMassISMRenderStateHel
 	// Hit proxies
 	if (RenderISMFragment.PerInstanceHitProxy.Num())
 	{
-		UE_LOG(LogTemp, Warning, TEXT("CreateProxies"));
 		const FMassEditorVisualizationMeshFragment* EditorMeshFragment = GetEditorMeshFragment();
 		const bool bSelected = EditorMeshFragment ? EditorMeshFragment->bShouldRenderSelected : false;
 
