Add optional FMassISMGroupingFragment for per-group ISM entity scoping:
Introduce FMassISMGroupingFragment, an optional const shared fragment carrying a FGuid GroupId. When present on source entities, the GroupId is included in the ISM entity hash, ensuring entities with different group scoping (e.g., per-cell, per-level) produce separate ISM instances even when they share the same mesh and visualization flags.

#Synced-CL 52109088

#Synced-CL 52109088

diff --git a/Engine/Source/Runtime/Mass/MassEngine/Private/Mesh/MassEngineRenderISMProcessors.cpp b/Engine/Source/Runtime/Mass/MassEngine/Private/Mesh/MassEngineRenderISMProcessors.cpp
--- a/Engine/Source/Runtime/Mass/MassEngine/Private/Mesh/MassEngineRenderISMProcessors.cpp
+++ b/Engine/Source/Runtime/Mass/MassEngine/Private/Mesh/MassEngineRenderISMProcessors.cpp
@@ -27,6 +27,7 @@ void UMassInstantiateStaticMeshAndCreateISMRenderStateProcessor::ConfigureQuerie
 	EntityQuery.AddTagRequirement<FMassRenderISMCandidateTag>(EMassFragmentPresence::All);
 	EntityQuery.AddConstSharedRequirement<FMassRenderISMLinkFragment>(EMassFragmentPresence::None);
 	EntityQuery.AddRequirement<FMassRenderISMInstanceFragment>(EMassFragmentAccess::ReadOnly, EMassFragmentPresence::None);
+	EntityQuery.AddConstSharedRequirement<FMassISMGroupingFragment>(EMassFragmentPresence::Optional);
 
 	ISMEntityComposition.Add<FTransformFragment>();
 	ISMEntityComposition.Add<FMassRenderStateFragment>();
@@ -79,6 +80,7 @@ void UMassInstantiateStaticMeshAndCreateISMRenderStateProcessor::Execute(FMassEn
 		const FMassVisualizationMeshFragment& MeshFragment,
 		const FMassOverrideMaterialsFragment* OverrideMaterialsFragment,
 		const FMassEditorVisualizationMeshFragment* EditorMeshFragment,
+		const FMassISMGroupingFragment* ISMGrouping,
 		bool bCreateIfNotFound,
 		FMassEntityHandle* ISMEntityHandle = nullptr) -> FISMEntityInfo*
 	{
@@ -97,9 +99,15 @@ void UMassInstantiateStaticMeshAndCreateISMRenderStateProcessor::Execute(FMassEn
 			Batch = &ISMEntityBatches.Emplace_GetRef(OverrideMaterialsFragment, EditorMeshFragment);
 		}
 
-		// Build the source shared fragment hash
+		// Build the source shared fragment hash.
+		// Include ISM grouping ID when present so entities with different group scoping
+		// (e.g., per-cell) get separate ISM entities even with the same mesh/vis flags.
 		FHashBuilder SourceSharedFragmentsHash;
 		SourceSharedFragmentsHash << &StaticMeshFragment << &MeshFragment;
+		if (ISMGrouping)
+		{
+			SourceSharedFragmentsHash << ISMGrouping->GroupId;
+		}
 
 		// Find if the ISM entity already exist, if not add it
 		FISMEntityInfo* ISMEntity;
@@ -141,7 +149,8 @@ void UMassInstantiateStaticMeshAndCreateISMRenderStateProcessor::Execute(FMassEn
 #endif // WITH_EDITOR
 
 			// Find the batch matching the shared fragment
-			FISMEntityInfo* ISMEntityInfo = RetrieveISMEntityInfo(StaticMeshFragment, MeshFragment, OverrideMaterialsFragment, EditorMeshFragment, /*bCreateIfNotFound*/true);
+			const FMassISMGroupingFragment* ISMGrouping = Context.GetConstSharedFragmentPtr<FMassISMGroupingFragment>();
+			FISMEntityInfo* ISMEntityInfo = RetrieveISMEntityInfo(StaticMeshFragment, MeshFragment, OverrideMaterialsFragment, EditorMeshFragment, ISMGrouping, /*bCreateIfNotFound*/true);
 			checkf(ISMEntityInfo, TEXT("Expecting the ISM to be created"));
 			ISMEntityInfo->InstanceCount += Context.GetNumEntities();
 		});
@@ -195,7 +204,8 @@ void UMassInstantiateStaticMeshAndCreateISMRenderStateProcessor::Execute(FMassEn
 #endif // WITH_EDITOR
 
 			FMassEntityHandle ISMEntityHandle;
-			FISMEntityInfo* ISMEntityInfo = RetrieveISMEntityInfo(StaticMeshFragment, MeshFragment, OverrideMaterialsFragment, EditorMeshFragment, /*bCreateIfNotFound*/false, &ISMEntityHandle);
+			const FMassISMGroupingFragment* ISMGrouping = Context.GetConstSharedFragmentPtr<FMassISMGroupingFragment>();
+			FISMEntityInfo* ISMEntityInfo = RetrieveISMEntityInfo(StaticMeshFragment, MeshFragment, OverrideMaterialsFragment, EditorMeshFragment, ISMGrouping, /*bCreateIfNotFound*/false, &ISMEntityHandle);
 			checkf(ISMEntityInfo, TEXT("Expecting the ISM to be created"));
 			FMassEntityView ISMEntityView(EntityManager, ISMEntityHandle);
 
diff --git a/Engine/Source/Runtime/Mass/MassEngine/Public/Mesh/MassEngineMeshFragments.h b/Engine/Source/Runtime/Mass/MassEngine/Public/Mesh/MassEngineMeshFragments.h
--- a/Engine/Source/Runtime/Mass/MassEngine/Public/Mesh/MassEngineMeshFragments.h
+++ b/Engine/Source/Runtime/Mass/MassEngine/Public/Mesh/MassEngineMeshFragments.h
@@ -309,4 +309,24 @@ struct FMassRenderISMCandidateTag : public FMassTag
 	GENERATED_BODY()
 };
 
+/**
+ * Optional const shared fragment for ISM grouping scoping.
+ * When present, its GroupId is included in the ISM hash, ensuring entities
+ * with different GroupIds get separate ISM entities even if they share the
+ * same mesh/vis flags. Use cases: per-cell scoping, per-level scoping, etc.
+ */
+USTRUCT()
+struct FMassISMGroupingFragment : public FMassConstSharedFragment
+{
+	GENERATED_BODY()
+
+	FMassISMGroupingFragment() = default;
+	explicit FMassISMGroupingFragment(const FGuid& InGroupId)
+		: GroupId(InGroupId)
+	{}
+
+	UPROPERTY(Transient, VisibleAnywhere, Category = Debug)
+	FGuid GroupId;
+};
+
 #undef UE_API
