Add Mass render state cleanup processors for static mesh and ISM entities when their fragments are removed
Refactor ISM observer processor to inherit from a new shared base class and extend its observed fragment types
Remove shared fragment support from mass_component in favor of const shared fragments only, and implement proper mesh fragment cleanup on RemoveFromWorld

#rnx
#Synced-CL 51802337

diff --git a/Engine/Source/Runtime/Mass/MassEngine/Private/Mesh/MassEngineRenderISMProcessors.cpp b/Engine/Source/Runtime/Mass/MassEngine/Private/Mesh/MassEngineRenderISMProcessors.cpp
--- a/Engine/Source/Runtime/Mass/MassEngine/Private/Mesh/MassEngineRenderISMProcessors.cpp
+++ b/Engine/Source/Runtime/Mass/MassEngine/Private/Mesh/MassEngineRenderISMProcessors.cpp
@@ -30,6 +30,7 @@ void UMassInstantiateStaticMeshAndCreateISMRenderStateProcessor::ConfigureQuerie
 	EntityQuery.AddConstSharedRequirement<FMassISMGroupingFragment>(EMassFragmentPresence::Optional);
 
 	ISMEntityComposition.Add<FTransformFragment>();
+	ISMEntityComposition.Add<FMassVisualizationMeshFragment>();
 	ISMEntityComposition.Add<FMassRenderStateFragment>();
 	ISMEntityComposition.Add<FMassRenderPrimitiveFragment>();
 	ISMEntityComposition.Add<FMassRenderISMFragment>();
@@ -184,7 +185,13 @@ void UMassInstantiateStaticMeshAndCreateISMRenderStateProcessor::Execute(FMassEn
 		}
 #endif // WITH_EDITOR
 
-		ExecutionContext.BatchCreateEntities(ArchetypeHandle, SharedFragmentValues, Batch.EntitiesInfo.Num(), Batch.CreatedEntities);
+		for (const FISMEntityInfo& EntityInfo : Batch.EntitiesInfo)
+		{
+			FMassArchetypeSharedFragmentValues EntitySharedFragmentValues(SharedFragmentValues);
+			EntitySharedFragmentValues.Add(FConstSharedStruct::Make(EntityInfo.MeshFragment));
+			EntitySharedFragmentValues.Sort();
+			Batch.CreatedEntities.Add(ExecutionContext.CreateEntity(ArchetypeHandle, EntitySharedFragmentValues));
+		}
 	}
 
 	//////////////////////////////////////////////////////////////////////////////
@@ -242,7 +249,7 @@ void UMassInstantiateStaticMeshAndCreateISMRenderStateProcessor::Execute(FMassEn
 						}
 						else
 						{
-							UE_LOG(LogMass, Log, TEXT("FMassRenderISMLinkFragment shared fragment not added: Entity is no longer valid"));
+							UE_LOGF(LogMassEngine, Log, "FMassRenderISMLinkFragment shared fragment not added: Entity is no longer valid");
 						}
 					});
 
@@ -361,7 +368,7 @@ void UMassDirtyRenderISMInstanceProcessor::SignalEntities(FMassEntityManager& En
 				{
 					FMassEntityHandle EntityHandle = Context.GetEntity(*EntityIt);
 					Context.Defer().RemoveFragment<FMassRenderISMInstanceFragment>(EntityHandle);
-					Context.Defer().PushCommand<FMassDeferredAddCommand>([EntityHandle](FMassEntityManager& InEntityManager)
+					Context.Defer().PushCommand<FMassDeferredRemoveCommand>([EntityHandle](FMassEntityManager& InEntityManager)
 						{
 							if (InEntityManager.IsEntityValid(EntityHandle))
 							{
@@ -369,7 +376,7 @@ void UMassDirtyRenderISMInstanceProcessor::SignalEntities(FMassEntityManager& En
 							}
 							else
 							{
-								UE_LOG(LogMass, Log, TEXT("FMassRenderISMLinkFragment shared fragment not added: Entity is no longer valid"));
+								UE_LOGF(LogMassEngine, Log, "FMassRenderISMLinkFragment shared fragment not added: Entity is no longer valid");
 							}
 						});
 
@@ -395,26 +402,26 @@ void UMassDirtyRenderISMInstanceProcessor::SignalEntities(FMassEntityManager& En
 }
 
 // ----------------------------------------------------------------------//
-// UMassRemoveInstanceAndDestroyISMRenderStateProcessor
+// UMassISMDestroyRenderStateProcessor
 //----------------------------------------------------------------------//
-UMassRemoveInstanceAndDestroyISMRenderStateProcessor::UMassRemoveInstanceAndDestroyISMRenderStateProcessor()
-	: EntityQuery(*this)
+UMassISMDestroyRenderStateProcessor::UMassISMDestroyRenderStateProcessor()
 {
-	ExecutionFlags = static_cast<int32>(EProcessorExecutionFlags::AllWorldModes);
 	ObservedTypes.Add(FMassRenderISMInstanceFragment::StaticStruct());
-	ObservedOperations = EMassObservedOperationFlags::Remove;
-	bRequiresGameThreadExecution = true; // @todo revisit if the execution is required on game thread
+	ObservedTypes.Add(FMassRenderISMLinkFragment::StaticStruct());
+	ObservedTypes.Add(FMassRenderISMCandidateTag::StaticStruct());
 }
 
-void UMassRemoveInstanceAndDestroyISMRenderStateProcessor::ConfigureQueries(const TSharedRef<FMassEntityManager>& EntityManager)
+void UMassISMDestroyRenderStateProcessor::ConfigureQueries(const TSharedRef<FMassEntityManager>& EntityManager)
 {
-	EntityQuery.AddRequirement<FMassRenderISMInstanceFragment>(EMassFragmentAccess::ReadOnly);
+	Super::ConfigureQueries(EntityManager);
+
+	EntityQuery.AddRequirement<FMassRenderISMInstanceFragment>(EMassFragmentAccess::ReadWrite);
 	EntityQuery.AddConstSharedRequirement<FMassRenderISMLinkFragment>();
 	EntityQuery.AddLinkedEntityRequirement<FMassRenderISMFragment, FMassRenderISMLinkFragment>(EMassFragmentAccess::ReadOnly);
 	EntityQuery.AddLinkedEntityRequirement<FMassRenderStateFragment, FMassRenderISMLinkFragment>(EMassFragmentAccess::ReadOnly);
 }
 
-void UMassRemoveInstanceAndDestroyISMRenderStateProcessor::Execute(FMassEntityManager& EntityManager, FMassExecutionContext& ExecutionContext)
+void UMassISMDestroyRenderStateProcessor::Execute(FMassEntityManager& EntityManager, FMassExecutionContext& ExecutionContext)
 {
 	FMassDestroyRenderStateContext BatchingContext(nullptr);
 
@@ -423,7 +430,7 @@ void UMassRemoveInstanceAndDestroyISMRenderStateProcessor::Execute(FMassEntityMa
 	EntityQuery.ForEachEntityChunk(ExecutionContext,
 		[&BatchingContext, &EntitiesToDestroy](FMassExecutionContext& Context)
 		{
-			const TConstArrayView<FMassRenderISMInstanceFragment> ISMInfoFragments = Context.GetFragmentView<FMassRenderISMInstanceFragment>();
+			const TArrayView<FMassRenderISMInstanceFragment> ISMInfoFragments = Context.GetMutableFragmentView<FMassRenderISMInstanceFragment>();
 			const FMassRenderISMLinkFragment LinkedISMEntityFragment = Context.GetConstSharedFragment<FMassRenderISMLinkFragment>();
 			FMassRenderISMFragment* ISMDescFragment = Context.GetIndirectFragmentPtr<FMassRenderISMFragment>(LinkedISMEntityFragment.LinkedEntityHandle);
 			checkf(ISMDescFragment, TEXT("Expecting a ISM desc fragment"));
@@ -433,6 +440,13 @@ void UMassRemoveInstanceAndDestroyISMRenderStateProcessor::Execute(FMassEntityMa
 			for (FMassExecutionContext::FEntityIterator EntityIt = Context.CreateEntityIterator(); EntityIt; ++EntityIt)
 			{
 				const int32 InstanceIndex = ISMInfoFragments[*EntityIt].InstanceIndex;
+
+				// @Todo: Temps fix for a bug in Mass that prevents us from adding right requirements which would prevent multiple observer call
+				// Check description of the problem here: UMassBaseStaticMeshDestroyRenderStateProcessor::ConfigureQueries
+				if (InstanceIndex == INDEX_NONE)
+				{
+					continue;
+				}
 				checkf(ISMDescFragment->PerInstanceSMData.IsValidIndex(InstanceIndex), TEXT("The entity is out of sync with its instance"));
 #if WITH_EDITOR
 				const int32 HitProxyIndex = ISMInfoFragments[*EntityIt].HitProxyIndex;
@@ -461,6 +475,23 @@ void UMassRemoveInstanceAndDestroyISMRenderStateProcessor::Execute(FMassEntityMa
 #endif // WITH_EDITOR
 					Context.Defer().RemoveTag<FMassSceneProxyCreatedTag>(LinkedISMEntityFragment.LinkedEntityHandle);
 				}
+
+				FMassEntityHandle EntityHandle = Context.GetEntity(*EntityIt);
+
+				ISMInfoFragments[*EntityIt].InstanceIndex = INDEX_NONE;
+				// @Todo to reenable once we allow to remove shared fragments through this command
+				Context.Defer().PushCommand<FMassCommandRemoveFragments<FMassRenderISMInstanceFragment/*, FMassRenderISMLinkFragment*/>>(EntityHandle);
+				Context.Defer().PushCommand<FMassDeferredRemoveCommand>([EntityHandle](FMassEntityManager& InEntityManager)
+					{
+						if (InEntityManager.IsEntityValid(EntityHandle))
+						{
+							InEntityManager.RemoveConstSharedFragmentFromEntity<FMassRenderISMLinkFragment>(EntityHandle);
+						}
+						else
+						{
+							UE_LOGF(LogMassEngine, Log, "FMassRenderISMLinkFragment not removed: Entity is no longer valid");
+						}
+					});
 			}
 		});
 
diff --git a/Engine/Source/Runtime/Mass/MassEngine/Private/Mesh/MassEngineRenderStateProcessors.cpp b/Engine/Source/Runtime/Mass/MassEngine/Private/Mesh/MassEngineRenderStateProcessors.cpp
--- a/Engine/Source/Runtime/Mass/MassEngine/Private/Mesh/MassEngineRenderStateProcessors.cpp
+++ b/Engine/Source/Runtime/Mass/MassEngine/Private/Mesh/MassEngineRenderStateProcessors.cpp
@@ -198,6 +198,7 @@ void UMassDestroyRenderStateProcessor::Execute(FMassEntityManager& EntityManager
 				}
 				Helper.DestroyRenderState(&BatchingContext);
 				RenderStateFragments[*EntityIt].DestroyRenderStateHelper();
+				Context.Defer().RemoveTag<FMassSceneProxyCreatedTag>(Context.GetEntity(*EntityIt));
 			}
 		});
 
diff --git a/Engine/Source/Runtime/Mass/MassEngine/Private/Mesh/MassEngineRenderStaticMeshProcessors.cpp b/Engine/Source/Runtime/Mass/MassEngine/Private/Mesh/MassEngineRenderStaticMeshProcessors.cpp
--- a/Engine/Source/Runtime/Mass/MassEngine/Private/Mesh/MassEngineRenderStaticMeshProcessors.cpp
+++ b/Engine/Source/Runtime/Mass/MassEngine/Private/Mesh/MassEngineRenderStaticMeshProcessors.cpp
@@ -193,4 +193,51 @@ void UMassStaticMeshUpdateRenderStateProcessor::SignalEntities(
 			}
 
 		});
-}
\ No newline at end of file
+}
+
+//----------------------------------------------------------------------//
+// UMassBaseStaticMeshDestroyRenderStateProcessor
+//----------------------------------------------------------------------//
+UMassBaseStaticMeshDestroyRenderStateProcessor::UMassBaseStaticMeshDestroyRenderStateProcessor()
+	: EntityQuery(*this)
+{
+	ExecutionFlags = static_cast<int32>(EProcessorExecutionFlags::AllWorldModes);
+	ObservedTypes.Add(FTransformFragment::StaticStruct());
+	ObservedTypes.Add(FMassStaticMeshFragment::StaticStruct());
+	ObservedTypes.Add(FMassVisualizationMeshFragment::StaticStruct());
+	ObservedOperations = EMassObservedOperationFlags::Remove;
+}
+
+void UMassBaseStaticMeshDestroyRenderStateProcessor::ConfigureQueries(const TSharedRef<FMassEntityManager>& EntityManager)
+{
+	// @Todo, this would prevent the observer to be called more than once when the fragments are removed in multiple calls, 
+	// but we need to sort out the Observer manager LockCount batcher before. Otherwise we get called after the fragments are 
+	// removed and the query will never run.
+	//EntityQuery.AddRequirement<FTransformFragment>(EMassFragmentAccess::ReadOnly);
+	//EntityQuery.AddConstSharedRequirement<FMassStaticMeshFragment>(EMassFragmentPresence::All);
+	//EntityQuery.AddConstSharedRequirement<FMassVisualizationMeshFragment>(EMassFragmentPresence::All);
+}
+
+//----------------------------------------------------------------------//
+// UMassStaticMeshDestroyRenderStateProcessor
+//----------------------------------------------------------------------//
+void UMassStaticMeshDestroyRenderStateProcessor::ConfigureQueries(const TSharedRef<FMassEntityManager>& EntityManager)
+{
+	Super::ConfigureQueries(EntityManager);
+
+	EntityQuery.AddRequirement<FMassRenderStateFragment>(EMassFragmentAccess::ReadOnly);
+	EntityQuery.AddRequirement<FMassRenderPrimitiveFragment>(EMassFragmentAccess::ReadOnly);
+	EntityQuery.AddRequirement<FMassRenderStaticMeshFragment>(EMassFragmentAccess::ReadOnly);
+}
+
+void UMassStaticMeshDestroyRenderStateProcessor::Execute(FMassEntityManager& EntityManager, FMassExecutionContext& ExecutionContext)
+{
+	EntityQuery.ForEachEntityChunk(ExecutionContext,
+		[](FMassExecutionContext& Context)
+		{
+			for (FMassExecutionContext::FEntityIterator EntityIt = Context.CreateEntityIterator(); EntityIt; ++EntityIt)
+			{
+				Context.Defer().PushCommand<FMassCommandRemoveFragments<FMassRenderStateFragment, FMassRenderPrimitiveFragment, FMassRenderStaticMeshFragment>>(Context.GetEntity(*EntityIt));
+			}
+		});
+}
diff --git a/Engine/Source/Runtime/Mass/MassEngine/Public/Mesh/MassEngineRenderISMProcessors.h b/Engine/Source/Runtime/Mass/MassEngine/Public/Mesh/MassEngineRenderISMProcessors.h
--- a/Engine/Source/Runtime/Mass/MassEngine/Public/Mesh/MassEngineRenderISMProcessors.h
+++ b/Engine/Source/Runtime/Mass/MassEngine/Public/Mesh/MassEngineRenderISMProcessors.h
@@ -58,19 +58,17 @@ protected:
 };
 
 /**
- * Processor responsible for removing instances and destroying the render state of the instanced static mesh entity
+ * Processor responsible for removing instances when the required fragments are being removed
  */
 UCLASS()
-class UMassRemoveInstanceAndDestroyISMRenderStateProcessor : public UMassObserverProcessor
+class UMassISMDestroyRenderStateProcessor : public UMassBaseStaticMeshDestroyRenderStateProcessor
 {
 	GENERATED_BODY()
 
 public:
-	UMassRemoveInstanceAndDestroyISMRenderStateProcessor();
+	UMassISMDestroyRenderStateProcessor();
 
 protected:
 	virtual void ConfigureQueries(const TSharedRef<FMassEntityManager>& EntityManager) override;
 	virtual void Execute(FMassEntityManager& EntityManager, FMassExecutionContext& Context) override;
-
-	FMassEntityQuery EntityQuery;
 };
\ No newline at end of file
diff --git a/Engine/Source/Runtime/Mass/MassEngine/Public/Mesh/MassEngineRenderStaticMeshProcessors.h b/Engine/Source/Runtime/Mass/MassEngine/Public/Mesh/MassEngineRenderStaticMeshProcessors.h
--- a/Engine/Source/Runtime/Mass/MassEngine/Public/Mesh/MassEngineRenderStaticMeshProcessors.h
+++ b/Engine/Source/Runtime/Mass/MassEngine/Public/Mesh/MassEngineRenderStaticMeshProcessors.h
@@ -62,3 +62,34 @@ protected:
 	virtual void ConfigureQueries(const TSharedRef<FMassEntityManager>& EntityManager) override;
 	virtual void SignalEntities(FMassEntityManager& EntityManager, FMassExecutionContext& Context, FMassSignalNameLookup& EntitySignals) override;
 };
+
+
+/**
+ * Base processor responsible for removing the common static mesh render state fragments when the required fragments are being removed
+ */
+UCLASS(Abstract)
+class UMassBaseStaticMeshDestroyRenderStateProcessor : public UMassObserverProcessor
+{
+	GENERATED_BODY()
+
+public:
+	UMassBaseStaticMeshDestroyRenderStateProcessor();
+
+	virtual void ConfigureQueries(const TSharedRef<FMassEntityManager>& EntityManager);
+
+protected:
+	FMassEntityQuery EntityQuery;
+};
+
+/**
+ * Processor responsible for removing the static mesh render state fragments when the required fragments are being removed
+ */
+UCLASS()
+class UMassStaticMeshDestroyRenderStateProcessor : public UMassBaseStaticMeshDestroyRenderStateProcessor
+{
+	GENERATED_BODY()
+
+protected:
+	virtual void ConfigureQueries(const TSharedRef<FMassEntityManager>& EntityManager) override;
+	virtual void Execute(FMassEntityManager& EntityManager, FMassExecutionContext& Context) override;
+};
