[MassTrace]
Extract MassInsights trace provider and analyzer to new MassDeveloper runtime module.

Move FMassTraceProvider, FMassTraceAnalyzer, IMassTraceProvider, IEditableMassTraceProvider,
and all supporting data types (FFragmentInfo, FArchetypeInfo, FEntityEventRecord, FRegion,
FRegionLane, etc.) from MassInsightsAnalysis plugin to Engine/Source/Runtime/Mass/MassDeveloper
under the UE::Mass::Trace namespace.
Port all MassInsightsUI consumers to the new types.
Simplify module entry point via SetupMassTraceAnalysis() factory function.

#Synced-CL 52294509

diff --git a/Engine/Plugins/MassInsights/Source/MassInsightsAnalysis/MassInsightsAnalysis.Build.cs b/Engine/Plugins/MassInsights/Source/MassInsightsAnalysis/MassInsightsAnalysis.Build.cs
--- a/Engine/Plugins/MassInsights/Source/MassInsightsAnalysis/MassInsightsAnalysis.Build.cs
+++ b/Engine/Plugins/MassInsights/Source/MassInsightsAnalysis/MassInsightsAnalysis.Build.cs
@@ -6,11 +6,17 @@ namespace UnrealBuildTool.Rules
 	{
 		public MassInsightsAnalysis(ReadOnlyTargetRules Target) : base(Target)
 		{
-			PublicDependencyModuleNames.AddRange(new string[] { "Core" });
+			PublicDependencyModuleNames.AddRange(
+			new string[]
+			{
+				"Core"
+			});
+
 			PrivateDependencyModuleNames.AddRange(
 			new string[]
 			{
 				"CoreUObject",
+				"MassDeveloper",
 				"TraceLog",
 				"TraceAnalysis",
 				"TraceServices",
diff --git a/Engine/Plugins/MassInsights/Source/MassInsightsAnalysis/Private/Analyzers/MassInsightsTraceAnalysis.cpp b/Engine/Plugins/MassInsights/Source/MassInsightsAnalysis/Private/Analyzers/MassInsightsTraceAnalysis.cpp
deleted file mode 100644
diff --git a/Engine/Plugins/MassInsights/Source/MassInsightsAnalysis/Private/Analyzers/MassInsightsTraceAnalysis.h b/Engine/Plugins/MassInsights/Source/MassInsightsAnalysis/Private/Analyzers/MassInsightsTraceAnalysis.h
deleted file mode 100644
diff --git a/Engine/Plugins/MassInsights/Source/MassInsightsAnalysis/Private/MassInsightsAnalysisModule.cpp b/Engine/Plugins/MassInsights/Source/MassInsightsAnalysis/Private/MassInsightsAnalysisModule.cpp
--- a/Engine/Plugins/MassInsights/Source/MassInsightsAnalysis/Private/MassInsightsAnalysisModule.cpp
+++ b/Engine/Plugins/MassInsights/Source/MassInsightsAnalysis/Private/MassInsightsAnalysisModule.cpp
@@ -2,20 +2,16 @@
 
 #include "MassInsightsAnalysisModule.h"
 
-#include "Model/MassInsightsPrivate.h"
-#include "Analyzers/MassInsightsTraceAnalysis.h"
 #include "Modules/ModuleManager.h"
 #include "Features/IModularFeatures.h"
+#include "Trace/MassTraceTypes.h"
 
 namespace MassInsightsAnalysis
 {
 
 void FMassInsightsAnalysisModule::OnAnalysisBegin(TraceServices::IAnalysisSession& InSession)
 {
-	TSharedPtr<FMassInsightsProvider> MassInsightsProvider = MakeShared<FMassInsightsProvider>(InSession);
-	InSession.AddProvider(GetMassInsightsProviderName(), MassInsightsProvider);
-
-	InSession.AddAnalyzer(new FMassInsightsTraceAnalyzer(InSession, *MassInsightsProvider));
+	UE::Mass::Trace::SetupMassTraceAnalysis(InSession);
 }
 
 void FMassInsightsAnalysisModule::GetModuleInfo(TraceServices::FModuleInfo& OutModuleInfo)
@@ -24,6 +20,11 @@ void FMassInsightsAnalysisModule::GetModuleInfo(TraceServices::FModuleInfo& OutM
 	OutModuleInfo.DisplayName = TEXT("MassInsights");
 }
 
+void FMassInsightsAnalysisModule::GetLoggers(TArray<const TCHAR*>& OutLoggers)
+{
+	OutLoggers.Add(TEXT("MassTrace"));
+}
+
 void FMassInsightsAnalysisModule::StartupModule()
 {
 	IModularFeatures::Get().RegisterModularFeature(TraceServices::ModuleFeatureName, this);
diff --git a/Engine/Plugins/MassInsights/Source/MassInsightsAnalysis/Private/MassInsightsAnalysisModule.h b/Engine/Plugins/MassInsights/Source/MassInsightsAnalysis/Private/MassInsightsAnalysisModule.h
--- a/Engine/Plugins/MassInsights/Source/MassInsightsAnalysis/Private/MassInsightsAnalysisModule.h
+++ b/Engine/Plugins/MassInsights/Source/MassInsightsAnalysis/Private/MassInsightsAnalysisModule.h
@@ -22,6 +22,7 @@ public:
 	public:
 	//~ Begin TraceServices::IModule interface
 	virtual void GetModuleInfo(TraceServices::FModuleInfo& OutModuleInfo) override;
+	virtual void GetLoggers(TArray<const TCHAR*>& OutLoggers) override;
 	virtual void OnAnalysisBegin(TraceServices::IAnalysisSession& Session) override;
 	//~ End TraceServices::IModule interface
 };
diff --git a/Engine/Plugins/MassInsights/Source/MassInsightsAnalysis/Private/Model/MassInsights.cpp b/Engine/Plugins/MassInsights/Source/MassInsightsAnalysis/Private/Model/MassInsights.cpp
deleted file mode 100644
diff --git a/Engine/Plugins/MassInsights/Source/MassInsightsAnalysis/Private/Model/MassInsightsPrivate.h b/Engine/Plugins/MassInsights/Source/MassInsightsAnalysis/Private/Model/MassInsightsPrivate.h
deleted file mode 100644
diff --git a/Engine/Plugins/MassInsights/Source/MassInsightsAnalysis/Public/MassInsightsAnalysis/Model/MassInsights.h b/Engine/Plugins/MassInsights/Source/MassInsightsAnalysis/Public/MassInsightsAnalysis/Model/MassInsights.h
deleted file mode 100644
diff --git a/Engine/Plugins/MassInsights/Source/MassInsightsUI/MassInsightsUI.Build.cs b/Engine/Plugins/MassInsights/Source/MassInsightsUI/MassInsightsUI.Build.cs
--- a/Engine/Plugins/MassInsights/Source/MassInsightsUI/MassInsightsUI.Build.cs
+++ b/Engine/Plugins/MassInsights/Source/MassInsightsUI/MassInsightsUI.Build.cs
@@ -17,6 +17,7 @@ namespace UnrealBuildTool.Rules
 				"TraceInsights",
 				"TraceInsightsCore",
 				"TraceServices",
+				"MassDeveloper",
 				"MassInsightsAnalysis",
 			});
 		}
diff --git a/Engine/Plugins/MassInsights/Source/MassInsightsUI/Private/MassInsightsTimingTrack.cpp b/Engine/Plugins/MassInsights/Source/MassInsightsUI/Private/MassInsightsTimingTrack.cpp
--- a/Engine/Plugins/MassInsights/Source/MassInsightsUI/Private/MassInsightsTimingTrack.cpp
+++ b/Engine/Plugins/MassInsights/Source/MassInsightsUI/Private/MassInsightsTimingTrack.cpp
@@ -18,6 +18,7 @@
 #include "Insights/ViewModels/TimingEvent.h"
 #include "MassInsightsUI/Widgets/SMassInsightsAnalysisTab.h"
 #include "Modules/ModuleManager.h"
+#include "Trace/MassTraceTypes.h"
 
 #define LOCTEXT_NAMESPACE "RegionsTimingTrack"
 
@@ -32,7 +33,7 @@ namespace MassInsightsUI
 ////////////////////////////////////////////////////////////////////////////////////////////////////
 
 FMassInsightsViewCommands::FMassInsightsViewCommands()
-: ::TCommands<FMassInsightsViewCommands>(
+: TCommands<FMassInsightsViewCommands>(
 	TEXT("FMassInsightsViewCommands"),
 	NSLOCTEXT("Contexts", "FMassInsightsViewCommands", "Insights - Timing View - Mass Processor"),
 	NAME_None,
@@ -167,7 +168,7 @@ void FMassInsightsTrack::InitTooltip(FTooltipDrawState& InOutTooltip, const ITim
 			double InFoundStartTime,
 			double InFoundEndTime,
 			uint32 InFoundDepth,
-			const MassInsightsAnalysis::FMassInsights& InRegion)
+			const UE::Mass::Trace::FRegion& InRegion)
 		{
 			InOutTooltip.Reset();
 			InOutTooltip.AddTitle(InRegion.Text, FLinearColor::White);
@@ -181,31 +182,38 @@ void FMassInsightsTrack::InitTooltip(FTooltipDrawState& InOutTooltip, const ITim
 void FMassInsightsTrack::BuildDrawState(ITimingEventsTrackDrawStateBuilder& Builder,
 	const ITimingTrackUpdateContext& Context)
 {
-	const FTimingTrackViewport& Viewport = Context.GetViewport();
+	using namespace UE::Mass::Trace;
 
-	const MassInsightsAnalysis::IMassInsightsProvider& RegionProvider = MassInsightsAnalysis::ReadMassInsightsProvider(*AnalysisSession);
-	TraceServices::FProviderReadScopeLock RegionProviderScopedLock(RegionProvider);
-
-	// whe're counting only non-empty lanes, so we can collapse empty ones in the visualization.
-	int32 CurDepth = 0;
-	RegionProvider.EnumerateLanes([this, Viewport, &CurDepth, &Builder](const MassInsightsAnalysis::FMassInsightsLane& Lane, const int32 Depth)
+	if (AnalysisSession.IsValid())
 	{
-		bool RegionHadEvents = false;
-		Lane.EnumerateRegions(Viewport.GetStartTime(), Viewport.GetEndTime(), [&Builder, &RegionHadEvents, &CurDepth](const MassInsightsAnalysis::FMassInsights& Region) -> bool
+		const FTimingTrackViewport& Viewport = Context.GetViewport();
+
+		const IMassTraceProvider& RegionProvider = ReadMassTraceProvider(*AnalysisSession);
+		TraceServices::FProviderReadScopeLock RegionProviderScopedLock(RegionProvider);
+
+		// we're counting only non-empty lanes, so we can collapse empty ones in the visualization.
+		int32 CurDepth = 0;
+		RegionProvider.EnumerateLanes([this, Viewport, &CurDepth, &Builder](const FRegionLane& Lane, const int32 Depth)
 		{
-			RegionHadEvents = true;
-			Builder.AddEvent(Region.BeginTime, Region.EndTime,CurDepth, Region.Text);
-			return true;
-		});
+			bool RegionHadEvents = false;
+			Lane.EnumerateRegions(Viewport.GetStartTime(), Viewport.GetEndTime(), [&Builder, &RegionHadEvents, &CurDepth](const FRegion& Region) -> bool
+			{
+				RegionHadEvents = true;
+				Builder.AddEvent(Region.BeginTime, Region.EndTime,CurDepth, Region.Text);
+				return true;
+			});
 
-		if (RegionHadEvents) CurDepth++;
-	});
+			if (RegionHadEvents) CurDepth++;
+		});
+	}
 }
 
 //////////////////////////////////////////////////////////////////////////
 
 void FMassInsightsTrack::BuildFilteredDrawState(ITimingEventsTrackDrawStateBuilder& Builder, const ITimingTrackUpdateContext& Context)
 {
+	using namespace UE::Mass::Trace;
+
 	const TSharedPtr<ITimingEventFilter> EventFilterPtr = Context.GetEventFilter();
 	if (EventFilterPtr.IsValid() && EventFilterPtr->FilterTrack(*this))
 	{
@@ -220,7 +228,7 @@ void FMassInsightsTrack::BuildFilteredDrawState(ITimingEventsTrackDrawStateBuild
 
 		if (AnalysisSession.IsValid())
 		{
-			const MassInsightsAnalysis::IMassInsightsProvider& RegionProvider = MassInsightsAnalysis::ReadMassInsightsProvider(*AnalysisSession);
+			const IMassTraceProvider& RegionProvider = ReadMassTraceProvider(*AnalysisSession);
 			TraceServices::FProviderReadScopeLock RegionProviderScopedLock(RegionProvider);
 
 			const FTimingTrackViewport& Viewport = Context.GetViewport();
@@ -228,10 +236,10 @@ void FMassInsightsTrack::BuildFilteredDrawState(ITimingEventsTrackDrawStateBuild
 			if (bFilterOnlyByEventType)
 			{
 				int32 CurDepth = 0;
-				RegionProvider.EnumerateLanes([this, Viewport, &CurDepth, &Builder, FilterEventType](const MassInsightsAnalysis::FMassInsightsLane& Lane, const int32 Depth)
+				RegionProvider.EnumerateLanes([this, Viewport, &CurDepth, &Builder, FilterEventType](const FRegionLane& Lane, const int32 Depth)
 					{
 						bool RegionHadEvents = false;
-						Lane.EnumerateRegions(Viewport.GetStartTime(), Viewport.GetEndTime(), [&Builder, &RegionHadEvents, &CurDepth, FilterEventType](const MassInsightsAnalysis::FMassInsights& Region) -> bool
+						Lane.EnumerateRegions(Viewport.GetStartTime(), Viewport.GetEndTime(), [&Builder, &RegionHadEvents, &CurDepth, FilterEventType](const FRegion& Region) -> bool
 							{
 								RegionHadEvents = true;
 								if (Region.Text == FilterEventType)
@@ -260,7 +268,7 @@ const TSharedPtr<const ITimingEvent> FMassInsightsTrack::SearchEvent(
 {
 	TSharedPtr<const ITimingEvent> FoundEvent;
 
-	FindRegionEvent(InSearchParameters, [this, &FoundEvent](double InFoundStartTime, double InFoundEndTime, uint32 InFoundDepth, const MassInsightsAnalysis::FMassInsights& InEvent)
+	FindRegionEvent(InSearchParameters, [this, &FoundEvent](double InFoundStartTime, double InFoundEndTime, uint32 InFoundDepth, const UE::Mass::Trace::FRegion& InEvent)
 	{
 		FoundEvent = MakeShared<FTimingEvent>(SharedThis(this), InFoundStartTime, InFoundEndTime, InFoundDepth, reinterpret_cast<uint64>(InEvent.Text));
 	});
@@ -271,25 +279,32 @@ const TSharedPtr<const ITimingEvent> FMassInsightsTrack::SearchEvent(
 //////////////////////////////////////////////////////////////////////////
 
 bool FMassInsightsTrack::FindRegionEvent(const FTimingEventSearchParameters& InParameters,
-	TFunctionRef<void(double, double, uint32, const MassInsightsAnalysis::FMassInsights&)> InFoundPredicate) const
+	TFunctionRef<void(double, double, uint32, const UE::Mass::Trace::FRegion&)> InFoundPredicate) const
 {
+	using namespace UE::Mass::Trace;
+
+	if (!AnalysisSession.IsValid())
+	{
+		return false;
+	}
+
 	// If the query start time is larger than the end of the session return false.
 	{
 		TraceServices::FAnalysisSessionReadScope SessionReadScope(*AnalysisSession.Get());
-		if (AnalysisSession.IsValid() && InParameters.StartTime > AnalysisSession->GetDurationSeconds())
+		if (InParameters.StartTime > AnalysisSession->GetDurationSeconds())
 		{
 			return false;
 		}
 	}
 
-	return TTimingEventSearch<MassInsightsAnalysis::FMassInsights>::Search(
+	return TTimingEventSearch<FRegion>::Search(
 		InParameters,
-		[this](TTimingEventSearch<MassInsightsAnalysis::FMassInsights>::FContext& InContext)
+		[this](TTimingEventSearch<FRegion>::FContext& InContext)
 		{
-			const MassInsightsAnalysis::IMassInsightsProvider& RegionProvider = MassInsightsAnalysis::ReadMassInsightsProvider(*AnalysisSession);
+			const IMassTraceProvider& RegionProvider = ReadMassTraceProvider(*AnalysisSession);
 			TraceServices::FProviderReadScopeLock RegionProviderScopedLock(RegionProvider);
 
-			RegionProvider.EnumerateRegions(InContext.GetParameters().StartTime, InContext.GetParameters().EndTime, [&InContext](const MassInsightsAnalysis::FMassInsights& Region)
+			RegionProvider.EnumerateRegions(InContext.GetParameters().StartTime, InContext.GetParameters().EndTime, [&InContext](const FRegion& Region)
 			{
 				InContext.Check(Region.BeginTime, Region.EndTime, Region.Depth, Region);
 
@@ -301,7 +316,7 @@ bool FMassInsightsTrack::FindRegionEvent(const FTimingEventSearchParameters& InP
 				return true;
 			});
 		},
-		TTimingEventSearch<MassInsightsAnalysis::FMassInsights>::NoMatch);
+		TTimingEventSearch<FRegion>::NoMatch);
 }
 
 //////////////////////////////////////////////////////////////////////////
diff --git a/Engine/Plugins/MassInsights/Source/MassInsightsUI/Private/MassInsightsTimingTrack.h b/Engine/Plugins/MassInsights/Source/MassInsightsUI/Private/MassInsightsTimingTrack.h
--- a/Engine/Plugins/MassInsights/Source/MassInsightsUI/Private/MassInsightsTimingTrack.h
+++ b/Engine/Plugins/MassInsights/Source/MassInsightsUI/Private/MassInsightsTimingTrack.h
@@ -4,11 +4,15 @@
 #include "Framework/Commands/Commands.h"
 #include "Insights/ITimingViewExtender.h"
 #include "Insights/ViewModels/TimingEventsTrack.h"
-#include "MassInsightsAnalysis/Model/MassInsights.h"
 #include "Templates/SharedPointer.h"
 
 ////////////////////////////////////////////////////////////////////////////////////////////////////
 
+namespace UE::Mass::Trace
+{
+struct FRegion;
+}
+
 namespace MassInsightsUI
 {
 
@@ -68,7 +72,7 @@ public:
 	virtual void SetFilterConfigurator(TSharedPtr<UE::Insights::FFilterConfigurator> InFilterConfigurator) override;
 
 protected:
-	bool FindRegionEvent(const FTimingEventSearchParameters& InParameters, TFunctionRef<void(double, double, uint32, const MassInsightsAnalysis::FMassInsights&)> InFoundPredicate) const;
+	bool FindRegionEvent(const FTimingEventSearchParameters& InParameters, TFunctionRef<void(double, double, uint32, const UE::Mass::Trace::FRegion&)> InFoundPredicate) const;
 
 private:
 	TSharedPtr<UE::Insights::FFilterConfigurator> FilterConfigurator;
diff --git a/Engine/Plugins/MassInsights/Source/MassInsightsUI/Private/MassInsightsUI/Widgets/SArchetypeDetails.cpp b/Engine/Plugins/MassInsights/Source/MassInsightsUI/Private/MassInsightsUI/Widgets/SArchetypeDetails.cpp
--- a/Engine/Plugins/MassInsights/Source/MassInsightsUI/Private/MassInsightsUI/Widgets/SArchetypeDetails.cpp
+++ b/Engine/Plugins/MassInsights/Source/MassInsightsUI/Private/MassInsightsUI/Widgets/SArchetypeDetails.cpp
@@ -1,4 +1,4 @@
-﻿// Copyright Epic Games, Inc. All Rights Reserved.
+// Copyright Epic Games, Inc. All Rights Reserved.
 
 #include "SArchetypeDetails.h"
 
@@ -8,8 +8,8 @@
 #include "Insights/IUnrealInsightsModule.h"
 #include "InsightsCore/Table/ViewModels/Table.h"
 #include "InsightsCore/Table/ViewModels/TableColumn.h"
-#include "MassInsightsAnalysis/Model/MassInsights.h"
 #include "Modules/ModuleManager.h"
+#include "Trace/MassTraceTypes.h"
 #include "Widgets/Input/SButton.h"
 
 #define LOCTEXT_NAMESPACE "MassInsightsUI::SArchetypeDetails"
@@ -24,13 +24,13 @@ namespace MassInsightsUI
 		using Super = SMultiColumnTableRow< TSharedPtr<SArchetypeDetails::FFragmentListEntry> >;
 	public:
 		SLATE_BEGIN_ARGS(SFragmentListRow) {}
-			SLATE_ARGUMENT(const MassInsightsAnalysis::FMassFragmentInfo*, FragmentInfo)
+			SLATE_ARGUMENT(const UE::Mass::Trace::FFragmentInfo*, FragmentInfo)
 		SLATE_END_ARGS()
 
 		void Construct(const FArguments& InArgs, const TSharedRef<STableViewBase>& InOwnerTableView);
 		TSharedRef<SWidget> GenerateWidgetForColumn(const FName& InColumnName) override;
 	private:
-		const MassInsightsAnalysis::FMassFragmentInfo* FragmentInfo = nullptr;
+		const UE::Mass::Trace::FFragmentInfo* FragmentInfo = nullptr;
 	};
 
 	void SFragmentListRow::Construct(const FArguments& InArgs, const TSharedRef<STableViewBase>& InOwnerTableView)
@@ -42,6 +42,8 @@ namespace MassInsightsUI
 
 	TSharedRef<SWidget> SFragmentListRow::GenerateWidgetForColumn(const FName& InColumnName)
 	{
+		using namespace UE::Mass::Trace;
+
 		if (InColumnName == NAME_COLUMN_FragmentName)
 		{
 			return SNew( STextBlock )
@@ -55,13 +57,13 @@ namespace MassInsightsUI
 					{
 						switch (FragmentInfo->Type)
 						{
-						case MassInsightsAnalysis::EFragmentType::Fragment: return LOCTEXT("Fragment", "Fragment");
-						case MassInsightsAnalysis::EFragmentType::Tag: return LOCTEXT("Tag", "Tag");
-						case MassInsightsAnalysis::EFragmentType::Shared: return LOCTEXT("Shared", "Shared");
-						case MassInsightsAnalysis::EFragmentType::Sparse: return LOCTEXT("Sparse", "Sparse");
-						case MassInsightsAnalysis::EFragmentType::ConstShared: return LOCTEXT("ConstShared", "Const Shared");
-						case MassInsightsAnalysis::EFragmentType::SparseTag: return LOCTEXT("SparseTag", "Sparse Tag");
-						case MassInsightsAnalysis::EFragmentType::Unknown: // Fallthrough
+						case EFragmentType::Fragment: return LOCTEXT("Fragment", "Fragment");
+						case EFragmentType::Tag: return LOCTEXT("Tag", "Tag");
+						case EFragmentType::Shared: return LOCTEXT("Shared", "Shared");
+						case EFragmentType::Sparse: return LOCTEXT("Sparse", "Sparse");
+						case EFragmentType::ConstShared: return LOCTEXT("ConstShared", "Const Shared");
+						case EFragmentType::SparseTag: return LOCTEXT("SparseTag", "Sparse Tag");
+						case EFragmentType::Unknown: // Fallthrough
 						default:
 							return LOCTEXT("Unknown", "Unknown");
 						}
@@ -71,13 +73,13 @@ namespace MassInsightsUI
 				{
 					switch (FragmentInfo->Type)
 					{
-						case MassInsightsAnalysis::EFragmentType::Fragment: return FColor(10, 120, 180);
-						case MassInsightsAnalysis::EFragmentType::Sparse: return FColor(10, 160, 140);
-						case MassInsightsAnalysis::EFragmentType::ConstShared: return FColor(140, 80, 180);
-						case MassInsightsAnalysis::EFragmentType::SparseTag: return FColor(120, 140, 10);
-						case MassInsightsAnalysis::EFragmentType::Tag: // Fallthrough
-						case MassInsightsAnalysis::EFragmentType::Shared: // Fallthrough
-						case MassInsightsAnalysis::EFragmentType::Unknown: // Fallthrough
+						case EFragmentType::Fragment: return FColor(10, 120, 180);
+						case EFragmentType::Sparse: return FColor(10, 160, 140);
+						case EFragmentType::ConstShared: return FColor(140, 80, 180);
+						case EFragmentType::SparseTag: return FColor(120, 140, 10);
+						case EFragmentType::Tag: // Fallthrough
+						case EFragmentType::Shared: // Fallthrough
+						case EFragmentType::Unknown: // Fallthrough
 						default:
 							return FColor(90, 120, 10);
 					}
@@ -181,18 +183,19 @@ namespace MassInsightsUI
 	{
 		SCompoundWidget::Tick(AllottedGeometry, InCurrentTime, InDeltaTime);
 
-		if (bArchetypeDirty)
+		using namespace UE::Mass::Trace;
+		if (bArchetypeDirty
+			&& Session.IsValid())
 		{
 			bArchetypeDirty = false;
 			bSortingDirty = true;
 			FragmentListData.Reset();
-			
 
-			const MassInsightsAnalysis::IMassInsightsProvider& Provider = MassInsightsAnalysis::ReadMassInsightsProvider(*Session);
+			const IMassTraceProvider& Provider = ReadMassTraceProvider(*Session);
 			{
 				TraceServices::FProviderReadScopeLock ProviderReadScopeLock(Provider);
 
-				const MassInsightsAnalysis::FMassArchetypeInfo* ArchetypeInfo = Provider.FindArchetypeById(ArchetypeID);
+				const FArchetypeInfo* ArchetypeInfo = Provider.FindArchetypeById(ArchetypeID);
 				if (ArchetypeInfo != nullptr)
 				{
 					for (int32 Index = 0; Index < ArchetypeInfo->Fragments.Num(); Index++)
diff --git a/Engine/Plugins/MassInsights/Source/MassInsightsUI/Private/MassInsightsUI/Widgets/SArchetypeDetails.h b/Engine/Plugins/MassInsights/Source/MassInsightsUI/Private/MassInsightsUI/Widgets/SArchetypeDetails.h
--- a/Engine/Plugins/MassInsights/Source/MassInsightsUI/Private/MassInsightsUI/Widgets/SArchetypeDetails.h
+++ b/Engine/Plugins/MassInsights/Source/MassInsightsUI/Private/MassInsightsUI/Widgets/SArchetypeDetails.h
@@ -1,10 +1,15 @@
-﻿// Copyright Epic Games, Inc. All Rights Reserved.
+// Copyright Epic Games, Inc. All Rights Reserved.
 
 #pragma once
-#include "MassInsightsAnalysis/Model/MassInsights.h"
+
 #include "Widgets/SCompoundWidget.h"
 #include "Widgets/Views/STreeView.h"
 
+namespace UE::Mass::Trace
+{
+struct FFragmentInfo;
+}
+
 namespace UE::Insights
 {
 	class FTable;
@@ -15,11 +20,6 @@ namespace TraceServices
 	class IAnalysisSession;
 }
 
-namespace MassInsightsAnalysis
-{
-	struct FMassFragmentInfo;
-}
-
 namespace MassInsightsUI
 {
 	class SArchetypeDetails : public SCompoundWidget
@@ -36,7 +36,7 @@ namespace MassInsightsUI
 		friend class SFragmentListRow;
 		struct FFragmentListEntry
 		{
-			const MassInsightsAnalysis::FMassFragmentInfo* Fragment = nullptr;
+			const UE::Mass::Trace::FFragmentInfo* Fragment = nullptr;
 		};
 
 		TSharedRef<ITableRow> TreeView_OnGenerateRow(TSharedPtr<FFragmentListEntry> NodePtr, const TSharedRef<STableViewBase>& OwnerTable);
diff --git a/Engine/Plugins/MassInsights/Source/MassInsightsUI/Private/MassInsightsUI/Widgets/SEntityEventAggregationTableView.cpp b/Engine/Plugins/MassInsights/Source/MassInsightsUI/Private/MassInsightsUI/Widgets/SEntityEventAggregationTableView.cpp
--- a/Engine/Plugins/MassInsights/Source/MassInsightsUI/Private/MassInsightsUI/Widgets/SEntityEventAggregationTableView.cpp
+++ b/Engine/Plugins/MassInsights/Source/MassInsightsUI/Private/MassInsightsUI/Widgets/SEntityEventAggregationTableView.cpp
@@ -1,4 +1,4 @@
-﻿// Copyright Epic Games, Inc. All Rights Reserved.
+// Copyright Epic Games, Inc. All Rights Reserved.
 
 #include "SEntityEventAggregationTableView.h"
 
@@ -10,6 +10,7 @@
 #include "InsightsCore/Table/ViewModels/TableCellValueSorter.h"
 #include "InsightsCore/Table/ViewModels/TableColumn.h"
 #include "Modules/ModuleManager.h"
+#include "Trace/MassTraceTypes.h"
 #include "Widgets/Input/SButton.h"
 #include "Widgets/Text/STextBlock.h"
 
@@ -426,6 +427,8 @@ namespace MassInsights
 
 	void SEntityEventAggregationTableView::IncrementalUpdate()
 	{
+		using namespace UE::Mass::Trace;
+
 		if (Session.IsValid())
 		{
 			uint64 PostUpdate_NextIncrementalEventIndex = EventsProcessed;
@@ -433,7 +436,7 @@ namespace MassInsights
 			const int32 PreUpdateRecordCount = EventAggregationRecords.Num();
 
 			// @TODO: Would be better if this were async instead of on Tick
-			const MassInsightsAnalysis::IMassInsightsProvider& Provider = MassInsightsAnalysis::ReadMassInsightsProvider(*Session);
+			const IMassTraceProvider& Provider = ReadMassTraceProvider(*Session);
 			{
 				TraceServices::FProviderReadScopeLock ProviderReadScopeLock(Provider);
 
@@ -444,13 +447,13 @@ namespace MassInsights
 				Provider.EnumerateEntityEvents(
 					EventsProcessed,
 					MaximumEventsToCheck,
-					[this, &PostUpdate_NextIncrementalEventIndex](const MassInsightsAnalysis::FMassEntityEventRecord& Event, uint64 EventIndex)
+					[this, &PostUpdate_NextIncrementalEventIndex](const FEntityEventRecord& Event, uint64 EventIndex)
 					{
 						++PostUpdate_NextIncrementalEventIndex;
 						
 						switch (Event.Operation)
 						{
-						case MassInsightsAnalysis::EMassEntityEventType::Created:
+						case EEntityEventType::Created:
 							{
 								const int32 RecordIndex = EventAggregationRecords.Num();
 								EntityToRecordIndexMap.Add(Event.Entity, EventAggregationRecords.Num());
@@ -466,7 +469,7 @@ namespace MassInsights
 									});
 							}
 							break;
-						case MassInsightsAnalysis::EMassEntityEventType::ArchetypeChange:
+						case EEntityEventType::ArchetypeChange:
 							{
 								int32* IndexPtr = EntityToRecordIndexMap.Find(Event.Entity);
 								if (IndexPtr != nullptr)
@@ -494,7 +497,7 @@ namespace MassInsights
 								}							
 							}
 							break;
-						case MassInsightsAnalysis::EMassEntityEventType::Destroyed:
+						case EEntityEventType::Destroyed:
 							{
 								int32* IndexPtr = EntityToRecordIndexMap.Find(Event.Entity);
 								if (IndexPtr != nullptr)
diff --git a/Engine/Plugins/MassInsights/Source/MassInsightsUI/Private/MassInsightsUI/Widgets/SEntityEventAggregationTableView.h b/Engine/Plugins/MassInsights/Source/MassInsightsUI/Private/MassInsightsUI/Widgets/SEntityEventAggregationTableView.h
--- a/Engine/Plugins/MassInsights/Source/MassInsightsUI/Private/MassInsightsUI/Widgets/SEntityEventAggregationTableView.h
+++ b/Engine/Plugins/MassInsights/Source/MassInsightsUI/Private/MassInsightsUI/Widgets/SEntityEventAggregationTableView.h
@@ -1,4 +1,4 @@
-﻿// Copyright Epic Games, Inc. All Rights Reserved.
+// Copyright Epic Games, Inc. All Rights Reserved.
 
 #pragma once
 #include "SFragmentTableRow.h"
@@ -11,6 +11,11 @@
 #include "InsightsCore/Table/ViewModels/Table.h"
 #include "MassInsightsUI/Widgets/Common.h"
 
+namespace TraceServices
+{
+class IAnalysisSession;
+}
+
 namespace MassInsights
 {
 	struct FEntityEventAggregateRecordHandle : public UE::Insights::FBaseTreeNode
diff --git a/Engine/Plugins/MassInsights/Source/MassInsightsUI/Private/MassInsightsUI/Widgets/SEntityEventsTableView.cpp b/Engine/Plugins/MassInsights/Source/MassInsightsUI/Private/MassInsightsUI/Widgets/SEntityEventsTableView.cpp
--- a/Engine/Plugins/MassInsights/Source/MassInsightsUI/Private/MassInsightsUI/Widgets/SEntityEventsTableView.cpp
+++ b/Engine/Plugins/MassInsights/Source/MassInsightsUI/Private/MassInsightsUI/Widgets/SEntityEventsTableView.cpp
@@ -1,4 +1,4 @@
-﻿// Copyright Epic Games, Inc. All Rights Reserved.
+// Copyright Epic Games, Inc. All Rights Reserved.
 
 #include "SEntityEventsTableView.h"
 
@@ -12,10 +12,11 @@
 #include "InsightsCore/Table/ViewModels/TableCellValueGetter.h"
 #include "InsightsCore/Table/ViewModels/TableCellValueSorter.h"
 #include "InsightsCore/Table/ViewModels/TableColumn.h"
+#include "MassInsightsUI/Widgets/Common.h"
 #include "Modules/ModuleManager.h"
+#include "Trace/MassTraceTypes.h"
 #include "Widgets/Input/SButton.h"
 #include "Widgets/Text/STextBlock.h"
-#include "MassInsightsUI/Widgets/Common.h"
 
 #define LOCTEXT_NAMESPACE "MassInsights::SEntityTimelineTableView"
 
@@ -67,10 +68,12 @@ namespace MassInsightsUI
 	
 	TSharedRef<SWidget> SEntityEventsTableRow::GenerateWidgetForColumn(const FName& InColumnName)
 	{
+		using namespace UE::Mass::Trace;
+
 		const uint64 ProviderEventIndex = EventNode->GetProviderEventIndex();
-		const MassInsightsAnalysis::IMassInsightsProvider& Provider = MassInsightsAnalysis::ReadMassInsightsProvider(*AnalysisSession);
+		const IMassTraceProvider& Provider = ReadMassTraceProvider(*AnalysisSession);
 
-		TValueOrError<MassInsightsAnalysis::FMassEntityEventRecord, void> Event = MakeError();
+		TValueOrError<FEntityEventRecord, void> Event = MakeError();
 		{
 			TraceServices::FProviderReadScopeLock ProviderReadScopeLock(Provider);
 
@@ -95,11 +98,11 @@ namespace MassInsightsUI
 			{
 				switch (EventType)
 				{
-				case MassInsightsAnalysis::EMassEntityEventType::Created:
+				case EEntityEventType::Created:
 					return FAppStyle::GetBrush("Icons.Plus");
-				case MassInsightsAnalysis::EMassEntityEventType::ArchetypeChange:
+				case EEntityEventType::ArchetypeChange:
 					return FAppStyle::GetBrush("Icons.ArrowRight");
-				case MassInsightsAnalysis::EMassEntityEventType::Destroyed:
+				case EEntityEventType::Destroyed:
 					return FAppStyle::GetBrush("Icons.X");
 				default:
 					return FAppStyle::GetBrush("Icons.Warning");
@@ -117,11 +120,11 @@ namespace MassInsightsUI
 					{
 						switch (EventType)
 						{
-						case MassInsightsAnalysis::EMassEntityEventType::Created:
+						case EEntityEventType::Created:
 							return FSlateColor(FColor(30.0f, 200.0f, 20.0f));
-						case MassInsightsAnalysis::EMassEntityEventType::ArchetypeChange:
+						case EEntityEventType::ArchetypeChange:
 							return FSlateColor(FColor(200.0f, 200.0f, 200.0f));
-						case MassInsightsAnalysis::EMassEntityEventType::Destroyed:
+						case EEntityEventType::Destroyed:
 							return FSlateColor(FColor(160.0f, 20.0f, 30.0f));
 						default:
 							return FSlateColor(FColor(255,255,255));
@@ -170,6 +173,8 @@ namespace MassInsightsUI
 
 	void SEntityEventsTableView::Construct(const FArguments& InArgs)
 	{
+		using namespace UE::Mass::Trace;
+
 		OnEntityEventSelected = InArgs._OnEntityEventSelected;
 		ChildSlot
 		[
@@ -211,7 +216,7 @@ namespace MassInsightsUI
 		];
 
 		IUnrealInsightsModule& UnrealInsightsModule = FModuleManager::LoadModuleChecked<IUnrealInsightsModule>("TraceInsights");
-		TSharedPtr<::Insights::IInsightsManager> InsightsManager = UnrealInsightsModule.GetInsightsManager();
+		TSharedPtr<Insights::IInsightsManager> InsightsManager = UnrealInsightsModule.GetInsightsManager();
 		if (InsightsManager.IsValid())
 		{
 			InsightsManager->GetSessionChangedEvent().AddSP(this, &SEntityEventsTableView::InsightsManager_OnSessionChanged);
@@ -244,17 +249,17 @@ namespace MassInsightsUI
 					const uint64 ProviderEventIndex = EventNode.GetProviderEventIndex();
 
 					IUnrealInsightsModule& UnrealInsightsModule = FModuleManager::LoadModuleChecked<IUnrealInsightsModule>("TraceInsights");
-					TSharedPtr<::Insights::IInsightsManager> InsightsManager = UnrealInsightsModule.GetInsightsManager();
+					TSharedPtr<Insights::IInsightsManager> InsightsManager = UnrealInsightsModule.GetInsightsManager();
 					if (InsightsManager.IsValid())
 					{
 						TSharedPtr<const TraceServices::IAnalysisSession> CurrentSession = UnrealInsightsModule.GetAnalysisSession();
 						if (CurrentSession.IsValid())
 						{
-							const MassInsightsAnalysis::IMassInsightsProvider& Provider = MassInsightsAnalysis::ReadMassInsightsProvider(*CurrentSession);
+							const IMassTraceProvider& Provider = ReadMassTraceProvider(*CurrentSession);
 			
 							{
 								TraceServices::FProviderReadScopeLock ProviderReadScopeLock(Provider);
-								TValueOrError<MassInsightsAnalysis::FMassEntityEventRecord, void> EventData = Provider.GetEntityEvent(ProviderEventIndex);
+								TValueOrError<FEntityEventRecord, void> EventData = Provider.GetEntityEvent(ProviderEventIndex);
 								if (EventData.HasValue())
 								{
 									return FTableCellValue(EventData.GetValue().Time);
@@ -332,9 +337,11 @@ namespace MassInsightsUI
 
 	void SEntityEventsTableView::IncrementalUpdate()
 	{
+		using namespace UE::Mass::Trace;
+
 		if (Session.IsValid())
 		{
-			const MassInsightsAnalysis::IMassInsightsProvider& Provider = MassInsightsAnalysis::ReadMassInsightsProvider(*Session);
+			const IMassTraceProvider& Provider = ReadMassTraceProvider(*Session);
 
 			// Note: Would be better if this were async instead of on Tick
 			{
@@ -345,7 +352,7 @@ namespace MassInsightsUI
 				const int32 PreviousEventNodeCount = Events.Num();
 
 				uint64 NewNextEventIndex = NextEventIndex;
-				Provider.EnumerateEntityEvents(NextEventIndex, MaximumEventsToProcess, [this, &NewNextEventIndex](const MassInsightsAnalysis::FMassEntityEventRecord& Event, uint64 ProviderEventIndex)
+				Provider.EnumerateEntityEvents(NextEventIndex, MaximumEventsToProcess, [this, &NewNextEventIndex](const FEntityEventRecord& Event, uint64 ProviderEventIndex)
 				{
 					++NewNextEventIndex;
 
diff --git a/Engine/Plugins/MassInsights/Source/MassInsightsUI/Private/MassInsightsUI/Widgets/SEntityEventsTableView.h b/Engine/Plugins/MassInsights/Source/MassInsightsUI/Private/MassInsightsUI/Widgets/SEntityEventsTableView.h
--- a/Engine/Plugins/MassInsights/Source/MassInsightsUI/Private/MassInsightsUI/Widgets/SEntityEventsTableView.h
+++ b/Engine/Plugins/MassInsights/Source/MassInsightsUI/Private/MassInsightsUI/Widgets/SEntityEventsTableView.h
@@ -1,4 +1,4 @@
-﻿// Copyright Epic Games, Inc. All Rights Reserved.
+// Copyright Epic Games, Inc. All Rights Reserved.
 
 #pragma once
 #include "Common.h"
@@ -11,6 +11,11 @@
 #include "Widgets/Views/SListView.h"
 #include "InsightsCore/Table/ViewModels/Table.h"
 
+namespace TraceServices
+{
+class IAnalysisSession;
+}
+
 namespace MassInsightsUI
 {
 	struct FOnSelectedEntityEventParams
diff --git a/Engine/Plugins/MassInsights/Source/MassInsightsUI/Private/MassInsightsUI/Widgets/SFragmentTableRow.h b/Engine/Plugins/MassInsights/Source/MassInsightsUI/Private/MassInsightsUI/Widgets/SFragmentTableRow.h
--- a/Engine/Plugins/MassInsights/Source/MassInsightsUI/Private/MassInsightsUI/Widgets/SFragmentTableRow.h
+++ b/Engine/Plugins/MassInsights/Source/MassInsightsUI/Private/MassInsightsUI/Widgets/SFragmentTableRow.h
@@ -1,11 +1,16 @@
-﻿// Copyright Epic Games, Inc. All Rights Reserved.
+// Copyright Epic Games, Inc. All Rights Reserved.
 
 #pragma once
+
 #include "InsightsCore/Table/ViewModels/Table.h"
-#include "MassInsightsAnalysis/Model/MassInsights.h"
 #include "Widgets/Views/STableRow.h"
 #include "Widgets/Views/SListView.h"
 
+namespace UE::Mass::Trace
+{
+enum class EFragmentType : uint8;
+}
+
 namespace MassInsights
 {
 	struct FMassFragmentRowNodeData
@@ -13,7 +18,7 @@ namespace MassInsights
 		uint64 Id;
 		FString Name;
 		uint32 Size;
-		MassInsightsAnalysis::EFragmentType Type;
+		UE::Mass::Trace::EFragmentType Type;
 	};
 	using MassFragmentInfoPtr = TSharedPtr<FMassFragmentRowNodeData>;
 	class SFragmentTableRow : public SMultiColumnTableRow< MassFragmentInfoPtr >
diff --git a/Engine/Plugins/MassInsights/Source/MassInsightsUI/Private/MassInsightsUI/Widgets/SFragmentTableView.cpp b/Engine/Plugins/MassInsights/Source/MassInsightsUI/Private/MassInsightsUI/Widgets/SFragmentTableView.cpp
--- a/Engine/Plugins/MassInsights/Source/MassInsightsUI/Private/MassInsightsUI/Widgets/SFragmentTableView.cpp
+++ b/Engine/Plugins/MassInsights/Source/MassInsightsUI/Private/MassInsightsUI/Widgets/SFragmentTableView.cpp
@@ -1,4 +1,4 @@
-﻿// Copyright Epic Games, Inc. All Rights Reserved.
+// Copyright Epic Games, Inc. All Rights Reserved.
 
 #include "SFragmentTableView.h"
 
@@ -9,6 +9,7 @@
 #include "InsightsCore/Table/ViewModels/TableCellValueSorter.h"
 #include "InsightsCore/Table/ViewModels/TableColumn.h"
 #include "Modules/ModuleManager.h"
+#include "Trace/MassTraceTypes.h"
 
 #define LOCTEXT_NAMESPACE "MassInsights::SFragmentTableView"
 
@@ -144,6 +145,8 @@ namespace MassInsights
 
 	void SFragmentTableView::RebuildTree(bool bResync)
 	{
+		using namespace UE::Mass::Trace;
+
 		if (bResync)
 		{
 			FragmentInfos.Reset();
@@ -152,7 +155,7 @@ namespace MassInsights
 		
 		if (Session.IsValid())
 		{
-			const MassInsightsAnalysis::IMassInsightsProvider& Provider = MassInsightsAnalysis::ReadMassInsightsProvider(*Session);
+			const IMassTraceProvider& Provider = ReadMassTraceProvider(*Session);
 
 			TraceServices::FProviderReadScopeLock ProviderReadScopeLock(Provider);
 			const uint32 FragmentCount = static_cast<uint32>(Provider.GetFragmentCount());
@@ -160,7 +163,7 @@ namespace MassInsights
 			if (FragmentCount != PreviousFragmentCount)
 			{
 				FragmentInfos.SetNum(FragmentCount);
-				Provider.EnumerateFragments([this](const MassInsightsAnalysis::FMassFragmentInfo& FragmentInfo, int32 Index)
+				Provider.EnumerateFragments([this](const FFragmentInfo& FragmentInfo, int32 Index)
 				{
 					MassFragmentInfoPtr InfoPtr = MakeShared<MassFragmentInfoPtr::ElementType>();
 					InfoPtr->Id = FragmentInfo.Id;
diff --git a/Engine/Plugins/MassInsights/Source/MassInsightsUI/Private/MassInsightsUI/Widgets/SFragmentTableView.h b/Engine/Plugins/MassInsights/Source/MassInsightsUI/Private/MassInsightsUI/Widgets/SFragmentTableView.h
--- a/Engine/Plugins/MassInsights/Source/MassInsightsUI/Private/MassInsightsUI/Widgets/SFragmentTableView.h
+++ b/Engine/Plugins/MassInsights/Source/MassInsightsUI/Private/MassInsightsUI/Widgets/SFragmentTableView.h
@@ -1,4 +1,4 @@
-﻿// Copyright Epic Games, Inc. All Rights Reserved.
+// Copyright Epic Games, Inc. All Rights Reserved.
 
 #pragma once
 #include "SFragmentTableRow.h"
@@ -6,6 +6,11 @@
 #include "Widgets/SCompoundWidget.h"
 #include "Widgets/Views/STreeView.h"
 
+namespace TraceServices
+{
+class IAnalysisSession;
+}
+
 namespace MassInsights
 {
 
diff --git a/Engine/Plugins/MassInsights/Source/MassInsightsUI/Private/MassInsightsUI/Widgets/SMassInsightsAnalysisTab.cpp b/Engine/Plugins/MassInsights/Source/MassInsightsUI/Private/MassInsightsUI/Widgets/SMassInsightsAnalysisTab.cpp
--- a/Engine/Plugins/MassInsights/Source/MassInsightsUI/Private/MassInsightsUI/Widgets/SMassInsightsAnalysisTab.cpp
+++ b/Engine/Plugins/MassInsights/Source/MassInsightsUI/Private/MassInsightsUI/Widgets/SMassInsightsAnalysisTab.cpp
@@ -1,4 +1,4 @@
-﻿// Copyright Epic Games, Inc. All Rights Reserved.
+// Copyright Epic Games, Inc. All Rights Reserved.
 
 #include "SMassInsightsAnalysisTab.h"
 
@@ -9,6 +9,7 @@
 #include "Insights/ITimingViewSession.h"
 #include "Insights/ViewModels/ITimingEvent.h"
 #include "Styling/AppStyle.h"
+#include "Trace/MassTraceTypes.h"
 #include "Widgets/SBoxPanel.h"
 #include "Widgets/Input/SCheckBox.h"
 #include "Widgets/Layout/SWidgetSwitcher.h"
@@ -22,6 +23,8 @@ namespace MassInsights
 {
 	void SMassInsightsAnalysisTab::Construct(const FArguments& InArgs)
 	{
+		using namespace UE::Mass::Trace;
+
 		TSharedPtr<SHorizontalBox> ViewModeSelectorsBox;
 		TSharedPtr<SWidgetSwitcher> ViewModeSwitcher;
 		
@@ -62,10 +65,11 @@ namespace MassInsights
 			.OnEntityEventSelected_Lambda([this](const MassInsightsUI::FOnSelectedEntityEventParams& Params)
 			{
 				TSharedPtr<MassInsightsUI::SArchetypeDetails> Pinned = ArchetypesDetails.Pin();
-				if (Pinned)
+				if (Pinned
+					&& AnalysisSession)
 				{
-					TValueOrError<MassInsightsAnalysis::FMassEntityEventRecord,void> Event = MakeError();
-					const MassInsightsAnalysis::IMassInsightsProvider& Provider = MassInsightsAnalysis::ReadMassInsightsProvider(*AnalysisSession);
+					TValueOrError<FEntityEventRecord,void> Event = MakeError();
+					const IMassTraceProvider& Provider = ReadMassTraceProvider(*AnalysisSession);
 					{
 						TraceServices::FProviderReadScopeLock ProviderReadScopeLock(Provider);
 
diff --git a/Engine/Source/Runtime/Mass/MassDeveloper/MassDeveloper.Build.cs b/Engine/Source/Runtime/Mass/MassDeveloper/MassDeveloper.Build.cs
new file mode 100644
--- /dev/null
+++ b/Engine/Source/Runtime/Mass/MassDeveloper/MassDeveloper.Build.cs
@@ -0,0 +1,40 @@
+// Copyright Epic Games, Inc. All Rights Reserved.
+
+namespace UnrealBuildTool.Rules
+{
+	public class MassDeveloper : ModuleRules
+	{
+		public static bool IsMassTraceAnalysisEnabled(ReadOnlyTargetRules Target)
+		{
+			return Target.Platform.IsInGroup(UnrealPlatformGroup.Desktop)
+				&& (Target.Type == TargetType.Editor || Target.Type == TargetType.Program);
+		}
+
+		public MassDeveloper(ReadOnlyTargetRules Target) : base(Target)
+		{
+			PublicDependencyModuleNames.AddRange(
+			new string[]
+			{
+				"Core",
+			});
+
+			if (IsMassTraceAnalysisEnabled(Target))
+			{
+				PublicDependencyModuleNames.Add("TraceServices");
+
+				PrivateDependencyModuleNames.AddRange(
+					new string[]
+					{
+						"TraceLog",
+						"TraceAnalysis",
+					});
+
+				PublicDefinitions.Add("UE_MASS_TRACE_ANALYSIS_ENABLED=1");
+			}
+			else
+			{
+				PublicDefinitions.Add("UE_MASS_TRACE_ANALYSIS_ENABLED=0");
+			}
+		}
+	}
+}
diff --git a/Engine/Source/Runtime/Mass/MassDeveloper/Private/MassDeveloperModule.cpp b/Engine/Source/Runtime/Mass/MassDeveloper/Private/MassDeveloperModule.cpp
new file mode 100644
--- /dev/null
+++ b/Engine/Source/Runtime/Mass/MassDeveloper/Private/MassDeveloperModule.cpp
@@ -0,0 +1,5 @@
+// Copyright Epic Games, Inc. All Rights Reserved.
+
+#include "Modules/ModuleManager.h"
+
+IMPLEMENT_MODULE(FDefaultModuleImpl, MassDeveloper);
diff --git a/Engine/Source/Runtime/Mass/MassDeveloper/Private/Trace/MassTraceAnalyzer.cpp b/Engine/Source/Runtime/Mass/MassDeveloper/Private/Trace/MassTraceAnalyzer.cpp
new file mode 100644
--- /dev/null
+++ b/Engine/Source/Runtime/Mass/MassDeveloper/Private/Trace/MassTraceAnalyzer.cpp
@@ -0,0 +1,187 @@
+// Copyright Epic Games, Inc. All Rights Reserved.
+
+#include "Trace/MassTraceAnalyzer.h"
+
+#if UE_MASS_TRACE_ANALYSIS_ENABLED
+
+#include "Common/ProviderLock.h"
+#include "HAL/LowLevelMemTracker.h"
+#include "Trace/MassTraceProvider.h"
+#include "Trace/MassTraceTypes.h"
+#include "TraceServices/Model/AnalysisSession.h"
+
+namespace UE::Mass::Trace
+{
+
+FMassTraceAnalyzer::FMassTraceAnalyzer(TraceServices::IAnalysisSession& InSession, TSharedRef<FMassTraceProvider> InProvider)
+	: Session(InSession)
+	, MassTraceProvider(InProvider)
+{
+}
+
+void FMassTraceAnalyzer::OnAnalysisBegin(const FOnAnalysisContext& Context)
+{
+	FInterfaceBuilder& Builder = Context.InterfaceBuilder;
+
+	Builder.RouteEvent(RouteId_RegisterMassFragment, "MassTrace", "RegisterMassFragment");
+	Builder.RouteEvent(RouteId_RegisterMassArchetype, "MassTrace", "RegisterMassArchetype");
+
+	Builder.RouteEvent(RouteId_MassBulkAddEntity, "MassTrace", "MassBulkAddEntity");
+	Builder.RouteEvent(RouteId_MassEntityMoved, "MassTrace", "MassEntityMoved");
+	Builder.RouteEvent(RouteId_MassBulkEntityDestroyed, "MassTrace", "MassBulkEntityDestroyed");
+
+	Builder.RouteEvent(RouteId_MassPhaseBegin, "MassTrace", "MassPhaseBegin");
+	Builder.RouteEvent(RouteId_MassPhaseEnd, "MassTrace", "MassPhaseEnd");
+}
+
+void FMassTraceAnalyzer::OnAnalysisEnd()
+{
+	TraceServices::FProviderEditScopeLock RegionProviderScopedLock(MassTraceProvider.Get());
+	MassTraceProvider.Get().OnAnalysisSessionEnded();
+}
+
+bool FMassTraceAnalyzer::OnEvent(uint16 RouteId, EStyle Style, const FOnEventContext& Context)
+{
+	const FEventData& EventData = Context.EventData;
+	switch (RouteId)
+	{
+	case RouteId_RegisterMassFragment:
+	{
+		const uint64 FragmentId = EventData.GetValue<uint64>("FragmentId");
+		FString FragmentName;
+		EventData.GetString("FragmentName", FragmentName);
+
+		const uint32 FragmentSize = EventData.GetValue<uint32>("FragmentSize");
+		const EFragmentType FragmentType = static_cast<EFragmentType>(EventData.GetValue<uint8>("FragmentType"));
+		const FFragmentInfo FragmentInfo
+		{
+			.Id = FragmentId,
+			.Name = MoveTemp(FragmentName),
+			.Size = FragmentSize,
+			.Type = FragmentType,
+		};
+
+		TraceServices::FProviderEditScopeLock Lock(MassTraceProvider.Get());
+
+		MassTraceProvider.Get().AddFragment(FragmentInfo);
+		break;
+	}
+	case RouteId_RegisterMassArchetype:
+	{
+		const uint64 Id = EventData.GetValue<uint64>("ArchetypeID");
+		const TArrayView<const uint64> FragmentArrayView = EventData.GetArrayView<uint64>("Fragments");
+		FArchetypeInfo ArchetypeInfo
+		{
+			.Id = Id,
+		};
+
+		ArchetypeInfo.Fragments.Reserve(FragmentArrayView.Num());
+		{
+			TraceServices::FProviderReadScopeLock ReadLock(MassTraceProvider.Get());
+			for (int32 Index = 0; Index < FragmentArrayView.Num(); ++Index)
+			{
+				const uint64 FragmentId = FragmentArrayView[Index];
+				if (const FFragmentInfo* Info = MassTraceProvider.Get().FindFragmentById(FragmentId))
+				{
+					ArchetypeInfo.Fragments.Add(Info);
+				}
+			}
+			Algo::Sort(ArchetypeInfo.Fragments, [](const FFragmentInfo* Lhs, const FFragmentInfo* Rhs)
+				{
+					if (Lhs->Type == Rhs->Type)
+					{
+						return Lhs->Name < Rhs->Name;
+					}
+					return Lhs->Type < Rhs->Type;
+				});
+		}
+		{
+			TraceServices::FProviderEditScopeLock EditLock(MassTraceProvider.Get());
+			MassTraceProvider.Get().AddArchetype(ArchetypeInfo);
+		}
+		break;
+	}
+	case RouteId_MassBulkAddEntity:
+	{
+		const uint64 Cycle = EventData.GetValue<uint64>("Cycle");
+		const TArrayView<const uint64> Entities = EventData.GetArrayView<uint64>("Entities");
+		const TArrayView<const uint64> ArchetypeIDs = EventData.GetArrayView<uint64>("ArchetypeIDs");
+
+		TraceServices::FProviderEditScopeLock EditLock(MassTraceProvider.Get());
+		MassTraceProvider.Get().BulkAddEntity(Context.EventTime.AsSeconds(Cycle), Entities, ArchetypeIDs);
+		break;
+	}
+	case RouteId_MassEntityMoved:
+	{
+		const uint64 Cycle = EventData.GetValue<uint64>("Cycle");
+		const uint64 Entity = EventData.GetValue<uint64>("Entity");
+		const uint64 Archetype = EventData.GetValue<uint64>("NewArchetypeID");
+
+		TraceServices::FProviderEditScopeLock EditLock(MassTraceProvider.Get());
+		MassTraceProvider.Get().BulkMoveEntity(
+			Context.EventTime.AsSeconds(Cycle),
+			MakeConstArrayView(&Entity, 1),
+			MakeConstArrayView(&Archetype, 1));
+		break;
+	}
+	case RouteId_MassBulkEntityDestroyed:
+	{
+		const uint64 Cycle = EventData.GetValue<uint64>("Cycle");
+		const TArrayView<const uint64> Entities = EventData.GetArrayView<uint64>("Entities");
+
+		TraceServices::FProviderEditScopeLock EditLock(MassTraceProvider.Get());
+		MassTraceProvider.Get().BulkDestroyEntity(
+			Context.EventTime.AsSeconds(Cycle),
+			Entities);
+		break;
+	}
+	case RouteId_MassPhaseBegin:
+	{
+		const uint64 Cycle = EventData.GetValue<uint64>("Cycle");
+		FString PhaseName;
+		EventData.GetString("PhaseName", PhaseName);
+
+		const uint64 PhaseId = EventData.GetValue<uint64>("PhaseId", 0);
+		// a missing or 0 RegionID indicates that the region is identified by name only
+		if (PhaseId > 0)
+		{
+			TraceServices::FProviderEditScopeLock RegionProviderScopedLock(MassTraceProvider.Get());
+			MassTraceProvider.Get().AppendRegionBegin(*PhaseName, PhaseId, Context.EventTime.AsSeconds(Cycle));
+		}
+		else
+		{
+			TraceServices::FProviderEditScopeLock RegionProviderScopedLock(MassTraceProvider.Get());
+			MassTraceProvider.Get().AppendRegionBegin(*PhaseName, Context.EventTime.AsSeconds(Cycle));
+		}
+
+		break;
+	}
+
+	case RouteId_MassPhaseEnd:
+	{
+		const uint64 Cycle = EventData.GetValue<uint64>("Cycle");
+		const uint64 PhaseID = EventData.GetValue<uint64>("PhaseId", 0);
+
+		TraceServices::FProviderEditScopeLock RegionProviderScopedLock(MassTraceProvider.Get());
+		if (PhaseID > 0)
+		{
+			MassTraceProvider.Get().AppendRegionEnd(PhaseID, Context.EventTime.AsSeconds(Cycle));
+		}
+		else
+		{
+			FString PhaseName = TEXT("Invalid");
+			EventData.GetString("PhaseName", PhaseName);
+			MassTraceProvider.Get().AppendRegionEnd(*PhaseName, Context.EventTime.AsSeconds(Cycle));
+		}
+		break;
+	}
+	default:
+		break;
+	}
+
+	return true;
+}
+
+} // namespace UE::Mass::Trace
+
+#endif // UE_MASS_TRACE_ANALYSIS_ENABLED
diff --git a/Engine/Source/Runtime/Mass/MassDeveloper/Private/Trace/MassTraceAnalyzer.h b/Engine/Source/Runtime/Mass/MassDeveloper/Private/Trace/MassTraceAnalyzer.h
new file mode 100644
--- /dev/null
+++ b/Engine/Source/Runtime/Mass/MassDeveloper/Private/Trace/MassTraceAnalyzer.h
@@ -0,0 +1,49 @@
+// Copyright Epic Games, Inc. All Rights Reserved.
+
+#pragma once
+
+#if UE_MASS_TRACE_ANALYSIS_ENABLED
+
+#include "Templates/SharedPointer.h"
+#include "Trace/Analyzer.h"
+
+namespace TraceServices
+{
+class IAnalysisSession;
+}
+
+namespace UE::Mass::Trace
+{
+class FMassTraceProvider;
+
+/**
+ * Analyzer for Mass trace events, parses the trace events from the 'MassTrace' Logger.
+ */
+class FMassTraceAnalyzer
+	: public UE::Trace::IAnalyzer
+{
+public:
+	FMassTraceAnalyzer(TraceServices::IAnalysisSession& InSession, TSharedRef<FMassTraceProvider> InProvider);
+	virtual void OnAnalysisBegin(const FOnAnalysisContext& Context) override;
+	virtual void OnAnalysisEnd() override;
+	virtual bool OnEvent(uint16 RouteId, EStyle Style, const FOnEventContext& Context) override;
+
+private:
+	enum : uint16
+	{
+		RouteId_RegisterMassFragment,
+		RouteId_RegisterMassArchetype,
+		RouteId_MassBulkAddEntity,
+		RouteId_MassEntityMoved,
+		RouteId_MassBulkEntityDestroyed,
+		RouteId_MassPhaseBegin,
+		RouteId_MassPhaseEnd
+	};
+
+	TraceServices::IAnalysisSession& Session;
+	TSharedRef<FMassTraceProvider> MassTraceProvider;
+};
+
+} // namespace UE::Mass::Trace
+
+#endif // UE_MASS_TRACE_ANALYSIS_ENABLED
diff --git a/Engine/Source/Runtime/Mass/MassDeveloper/Private/Trace/MassTraceProvider.cpp b/Engine/Source/Runtime/Mass/MassDeveloper/Private/Trace/MassTraceProvider.cpp
new file mode 100644
--- /dev/null
+++ b/Engine/Source/Runtime/Mass/MassDeveloper/Private/Trace/MassTraceProvider.cpp
@@ -0,0 +1,450 @@
+// Copyright Epic Games, Inc. All Rights Reserved.
+
+#include "MassTraceProvider.h"
+
+#if UE_MASS_TRACE_ANALYSIS_ENABLED
+
+#include "Algo/ForEach.h"
+#include "Internationalization/Internationalization.h"
+#include "Templates/ValueOrError.h"
+#include "Trace/MassTraceAnalyzer.h"
+
+DEFINE_LOG_CATEGORY(LogMassTrace)
+
+namespace UE::Mass::Trace
+{
+
+thread_local TraceServices::FProviderLock::FThreadLocalState GMassInsightsProviderLockState;
+
+FMassTraceProvider::FMassTraceProvider(TraceServices::IAnalysisSession& InSession)
+	: Session(InSession)
+	, EntityEvents(InSession.GetLinearAllocator(), Constants::EntityEventsPageSize)
+{
+}
+
+int32 FMassTraceProvider::GetFragmentCount() const
+{
+	ReadAccessCheck();
+	return FragmentInfoByID.Num();
+}
+
+const FFragmentInfo* FMassTraceProvider::FindFragmentById(const uint64 FragmentId) const
+{
+	ReadAccessCheck();
+
+	const FFragmentInfo* const* Result = FragmentInfoByID.Find(FragmentId);
+	if (Result != nullptr)
+	{
+		return *Result;
+	}
+	return nullptr;
+}
+
+const FArchetypeInfo* FMassTraceProvider::FindArchetypeById(const uint64 ArchetypeId) const
+{
+	ReadAccessCheck();
+
+	const FArchetypeInfo* const* Result = ArchetypeByID.Find(ArchetypeId);
+	if (Result != nullptr)
+	{
+		return *Result;
+	}
+	return nullptr;
+}
+
+void FMassTraceProvider::EnumerateFragments(const TFunctionRef<void(const FFragmentInfo& FragmentInfo, int32 Index)> Callback, const int32 BeginIndex) const
+{
+	ReadAccessCheck();
+	for (int32 Index = BeginIndex, End = FragmentInfos.Num(); Index < End; Index++)
+	{
+		const FFragmentInfo& FragmentInfo = FragmentInfos[Index];
+		Callback(FragmentInfo, Index);
+	}
+}
+
+uint64 FMassTraceProvider::GetEntityEventCount() const
+{
+	ReadAccessCheck();
+
+	return EntityEvents.Num();
+}
+
+TValueOrError<FEntityEventRecord, void> FMassTraceProvider::GetEntityEvent(const uint64 EventIndex) const
+{
+	ReadAccessCheck();
+
+	if (EventIndex < EntityEvents.Num())
+	{
+		return MakeValue(EntityEvents[EventIndex]);
+	}
+
+	return MakeError();
+}
+
+void FMassTraceProvider::EnumerateEntityEvents(
+	const uint64 StartIndex,
+	const uint64 Count,
+	const TFunctionRef<void(const FEntityEventRecord&, uint64 EventIndex)> Callback) const
+{
+	ReadAccessCheck();
+
+	const uint64 EndIndex = FPlatformMath::Min(StartIndex + Count, EntityEvents.Num());
+	for (uint64 Index = StartIndex; Index < EndIndex; Index++)
+	{
+		Callback(EntityEvents[Index], Index);
+	}
+}
+
+uint64 FMassTraceProvider::GetRegionCount() const
+{
+	ReadAccessCheck();
+
+	uint64 RegionCount = 0;
+	for (const FRegionLane& Lane : Lanes)
+	{
+		RegionCount += Lane.Num();
+	}
+	return RegionCount;
+}
+
+int32 FMassTraceProvider::GetLaneCount() const
+{
+	ReadAccessCheck();
+	return Lanes.Num();
+}
+
+const FRegionLane* FMassTraceProvider::GetLane(const int32 Index) const
+{
+	ReadAccessCheck();
+
+	if (Index < Lanes.Num())
+	{
+		return &(Lanes[Index]);
+	}
+	return nullptr;
+}
+
+void FMassTraceProvider::AppendRegionBegin(const TCHAR* Name, const double Time)
+{
+	EditAccessCheck();
+	// Regions identified by Name don't have an ID
+	AppendRegionBegin(Name, 0, Time);
+}
+
+void FMassTraceProvider::AddFragment(const FFragmentInfo& FragmentInfo)
+{
+	EditAccessCheck();
+
+	if (FragmentInfoByID.Find(FragmentInfo.Id) == nullptr)
+	{
+		const uint64 Id = FragmentInfo.Id;
+		const int32 AllocatedIndex = FragmentInfos.AddElement(FragmentInfo);
+		const FFragmentInfo* AllocatedFragment = &FragmentInfos[AllocatedIndex];
+		FragmentInfoByID.Add(Id, AllocatedFragment);
+	}
+}
+
+void FMassTraceProvider::AddArchetype(const FArchetypeInfo& ArchetypeInfo)
+{
+	EditAccessCheck();
+
+	if (ArchetypeByID.Find(ArchetypeInfo.Id) == nullptr)
+	{
+		const uint64 Id = ArchetypeInfo.Id;
+		const int32 AllocatedIndex = ArchetypeInfos.AddElement(ArchetypeInfo);
+		const FArchetypeInfo* AllocatedArchetype = &ArchetypeInfos[AllocatedIndex];
+		ArchetypeByID.Add(Id, AllocatedArchetype);
+	}
+}
+
+void FMassTraceProvider::BulkAddEntity(const double Time, const TConstArrayView<uint64> Entities, const TConstArrayView<uint64> ArchetypeIDs)
+{
+	EditAccessCheck();
+
+	FEntityEventRecord Row;
+	Row.Time = Time;
+	Row.Operation = EEntityEventType::Created;
+
+	for (int32 Index = 0, End = Entities.Num(); Index < End; Index++)
+	{
+		Row.ArchetypeID = ArchetypeIDs[Index];
+		Row.Entity = Entities[Index];
+		EntityEvents.EmplaceBack(Row);
+	}
+}
+
+void FMassTraceProvider::BulkMoveEntity(const double Time, const TConstArrayView<uint64> Entities, const TConstArrayView<uint64> ArchetypeIDs)
+{
+	EditAccessCheck();
+
+	FEntityEventRecord Row;
+	Row.Time = Time;
+	Row.Operation = EEntityEventType::ArchetypeChange;
+
+	for (int32 Index = 0, End = Entities.Num(); Index < End; Index++)
+	{
+		Row.ArchetypeID = ArchetypeIDs[Index];
+		Row.Entity = Entities[Index];
+		EntityEvents.EmplaceBack(Row);
+	}
+}
+
+void FMassTraceProvider::BulkDestroyEntity(const double Time, const TConstArrayView<uint64> Entities)
+{
+	EditAccessCheck();
+
+	FEntityEventRecord Row;
+	Row.Time = Time;
+	Row.Operation = EEntityEventType::Destroyed;
+	Row.ArchetypeID = 0;
+
+	for (int32 Index = 0, End = Entities.Num(); Index < End; Index++)
+	{
+		Row.Entity = Entities[Index];
+		EntityEvents.EmplaceBack(Row);
+	}
+}
+
+void FMassTraceProvider::AppendRegionBegin(const TCHAR* Name, const uint64 ID, const double Time)
+{
+	EditAccessCheck();
+	// lookup by ID if the region has an ID
+	FRegion** OpenRegion = ID ? OpenRegionsByID.Find(ID) : OpenRegionsByName.Find(Name);
+
+	if (OpenRegion)
+	{
+		++NumWarnings;
+		if (NumWarnings <= MaxWarningMessages)
+		{
+			UE_LOG(LogMassTrace, Warning, TEXT("[Regions] A region begin event (%s) was encountered while a region with same name is already open."), Name)
+		}
+	}
+	else
+	{
+		FRegion Region;
+		Region.BeginTime = Time;
+		Region.Text = Session.StoreString(Name);
+		Region.ID = ID;
+		Region.Depth = CalculateRegionDepth(Region);
+
+		if (Region.Depth == Lanes.Num())
+		{
+			Lanes.Emplace(Session.GetLinearAllocator());
+		}
+
+		Lanes[Region.Depth].Regions.EmplaceBack(Region);
+		FRegion* NewOpenRegion = &(Lanes[Region.Depth].Regions.Last());
+
+		if (ID)
+		{
+			OpenRegionsByID.Add(Region.ID, NewOpenRegion);
+		}
+		else
+		{
+			OpenRegionsByName.Add(Region.Text, NewOpenRegion);
+		}
+		UpdateCounter++;
+	}
+
+	// Update session time
+	{
+		TraceServices::FAnalysisSessionEditScope SessionEditScope(Session);
+		Session.UpdateDurationSeconds(Time);
+	}
+}
+
+void FMassTraceProvider::AppendRegionEnd(const uint64 ID, const double Time)
+{
+	EditAccessCheck();
+
+	FRegion** OpenRegionPos = OpenRegionsByID.Find(ID);
+	FRegion* OpenRegion = nullptr;
+	if (OpenRegionPos)
+	{
+		OpenRegion = *OpenRegionPos;
+	}
+
+	AppendRegionEndInternal(OpenRegion, Time);
+}
+
+void FMassTraceProvider::AppendRegionEnd(const TCHAR* Name, const double Time)
+{
+	EditAccessCheck();
+
+	FRegion** OpenRegionPos = OpenRegionsByName.Find(Name);
+	FRegion* OpenRegion = nullptr;
+	if (OpenRegionPos)
+	{
+		OpenRegion = *OpenRegionPos;
+	}
+	else
+	{
+		++NumWarnings;
+		if (NumWarnings <= MaxWarningMessages)
+		{
+			UE_LOG(LogMassTrace, Warning, TEXT("[Regions] A region end event (%s) was encountered without having seen a matching region start event first."), Name)
+		}
+	}
+
+	AppendRegionEndInternal(OpenRegion, Time);
+}
+
+void FMassTraceProvider::AppendRegionEndInternal(FRegion* OpenRegion, const double Time)
+{
+	if (OpenRegion)
+	{
+		OpenRegion->EndTime = Time;
+
+		if (OpenRegion->ID)
+		{
+			OpenRegionsByID.Remove(OpenRegion->ID);
+		}
+		OpenRegionsByName.Remove(OpenRegion->Text);
+		UpdateCounter++;
+	}
+
+	// Update session time
+	{
+		TraceServices::FAnalysisSessionEditScope SessionEditScope(Session);
+		Session.UpdateDurationSeconds(Time);
+	}
+}
+
+void FMassTraceProvider::OnAnalysisSessionEnded()
+{
+	EditAccessCheck();
+
+	auto PrintOpenRegionMessage = [this](const TPair<uint64, FRegion*>& Pair)
+		{
+			const FRegion* Region = Pair.Value;
+			++NumWarnings;
+			if (NumWarnings <= MaxWarningMessages)
+			{
+				UE_LOG(LogMassTrace, Warning, TEXT("[Regions] A region begin event (%s) was never closed."), Region->Text)
+			}
+		};
+	Algo::ForEach(OpenRegionsByID, PrintOpenRegionMessage);
+	Algo::ForEach(OpenRegionsByName, [this](const TPair<FStringView, FRegion*>& Pair)
+		{
+			const FRegion* Region = Pair.Value;
+			++NumWarnings;
+			if (NumWarnings <= MaxWarningMessages)
+			{
+				UE_LOG(LogMassTrace, Warning, TEXT("[Regions] A region begin event (%s) was never closed."), Region->Text)
+			}
+		});
+
+	if (NumWarnings > 0 || NumErrors > 0)
+	{
+		UE_LOG(LogMassTrace, Error, TEXT("[Regions] %u warnings; %u errors"), NumWarnings, NumErrors);
+	}
+
+	const uint64 TotalRegionCount = GetRegionCount();
+	UE_LOG(LogMassTrace, Log, TEXT("[Regions] Analysis completed (%llu regions, %d lanes)."), TotalRegionCount, Lanes.Num());
+}
+
+int32 FMassTraceProvider::CalculateRegionDepth(const FRegion& Region) const
+{
+	constexpr int32 DepthLimit = 100;
+
+	int32 NewDepth = 0;
+
+	// Find first free lane/depth
+	while (NewDepth < DepthLimit)
+	{
+		if (!Lanes.IsValidIndex(NewDepth))
+		{
+			break;
+		}
+
+		const FRegion& LastRegion = Lanes[NewDepth].Regions.Last();
+		if (LastRegion.EndTime <= Region.BeginTime)
+		{
+			break;
+		}
+		NewDepth++;
+	}
+
+	ensureMsgf(NewDepth < DepthLimit, TEXT("Regions are nested too deep."));
+
+	return NewDepth;
+}
+
+void FMassTraceProvider::EnumerateLanes(const TFunctionRef<void(const FRegionLane&, int32)> Callback) const
+{
+	ReadAccessCheck();
+
+	for (int32 LaneIndex = 0; LaneIndex < Lanes.Num(); ++LaneIndex)
+	{
+		Callback(Lanes[LaneIndex], LaneIndex);
+	}
+}
+
+bool FMassTraceProvider::EnumerateRegions(const double IntervalStart, const double IntervalEnd, const TFunctionRef<bool(const FRegion&)> Callback) const
+{
+	ReadAccessCheck();
+
+	if (IntervalStart > IntervalEnd)
+	{
+		return false;
+	}
+
+	for (const FRegionLane& Lane : Lanes)
+	{
+		if (!Lane.EnumerateRegions(IntervalStart, IntervalEnd, Callback))
+		{
+			return false;
+		}
+	}
+
+	return true;
+}
+
+bool FRegionLane::EnumerateRegions(const double IntervalStart, const double IntervalEnd, TFunctionRef<bool(const FRegion&)> Callback) const
+{
+	const FInt32Interval OverlapRange = GetElementRangeOverlappingGivenRange<FRegion>(Regions, IntervalStart, IntervalEnd,
+		[](const FRegion& Region) { return Region.BeginTime; },
+		[](const FRegion& Region) { return Region.EndTime; });
+
+	if (OverlapRange.Min == -1)
+	{
+		return true;
+	}
+
+	for (int32 Index = OverlapRange.Min; Index <= OverlapRange.Max; ++Index)
+	{
+		if (!Callback(Regions[Index]))
+		{
+			return false;
+		}
+	}
+
+	return true;
+}
+
+FName GetProviderName()
+{
+	static const FName Name("MassTraceProvider");
+	return Name;
+}
+
+void SetupMassTraceAnalysis(TraceServices::IAnalysisSession& InSession)
+{
+	const TSharedPtr<FMassTraceProvider> MassProvider = MakeShared<FMassTraceProvider>(InSession);
+	InSession.AddProvider(GetProviderName(), MassProvider);
+	InSession.AddAnalyzer(MakeShared<FMassTraceAnalyzer>(InSession, MassProvider.ToSharedRef()));
+}
+
+const IMassTraceProvider& ReadMassTraceProvider(const TraceServices::IAnalysisSession& Session)
+{
+	return *Session.ReadProvider<IMassTraceProvider>(GetProviderName());
+}
+
+IEditableMassTraceProvider& EditMassTraceProvider(TraceServices::IAnalysisSession& Session)
+{
+	return *Session.EditProvider<IEditableMassTraceProvider>(GetProviderName());
+}
+
+} // namespace UE::Mass::Trace
+
+#endif // UE_MASS_TRACE_ANALYSIS_ENABLED
diff --git a/Engine/Source/Runtime/Mass/MassDeveloper/Private/Trace/MassTraceProvider.h b/Engine/Source/Runtime/Mass/MassDeveloper/Private/Trace/MassTraceProvider.h
new file mode 100644
--- /dev/null
+++ b/Engine/Source/Runtime/Mass/MassDeveloper/Private/Trace/MassTraceProvider.h
@@ -0,0 +1,152 @@
+// Copyright Epic Games, Inc. All Rights Reserved.
+
+#pragma once
+
+#if UE_MASS_TRACE_ANALYSIS_ENABLED
+
+#include "Common/ProviderLock.h"
+#include "Containers/ChunkedArray.h"
+#include "Containers/PagedArray.h"
+#include "Templates/SharedPointer.h"
+#include "Trace/MassTraceTypes.h"
+
+DECLARE_LOG_CATEGORY_EXTERN(LogMassTrace, Display, All)
+
+namespace TraceServices
+{
+class FAnalysisSessionLock;
+class FStringStore;
+}
+
+namespace UE::Mass::Trace
+{
+
+extern thread_local TraceServices::FProviderLock::FThreadLocalState GMassInsightsProviderLockState;
+
+class FMassTraceProvider
+	: public IMassTraceProvider
+	, public IEditableMassTraceProvider
+{
+	struct Constants
+	{
+		static constexpr uint64 EntityEventsPageSize = 65536;
+	};
+public:
+	explicit FMassTraceProvider(TraceServices::IAnalysisSession& Session);
+	virtual ~FMassTraceProvider() override = default;
+
+	//~ Read operations
+
+	virtual void BeginRead() const override
+	{
+		Lock.BeginRead(GMassInsightsProviderLockState);
+	}
+	virtual void EndRead() const override
+	{
+		Lock.EndRead(GMassInsightsProviderLockState);
+	}
+	virtual void ReadAccessCheck() const override
+	{
+		Lock.ReadAccessCheck(GMassInsightsProviderLockState);
+	}
+
+	virtual int32 GetFragmentCount() const override;
+	virtual const FFragmentInfo* FindFragmentById(uint64 FragmentId) const override;
+	virtual const FArchetypeInfo* FindArchetypeById(uint64 ArchetypeId) const override;
+
+	virtual void EnumerateFragments(TFunctionRef<void(const FFragmentInfo& FragmentInfo, int32 Index)> Callback, int32 BeginIndex = 0) const override;
+
+	virtual uint64 GetEntityEventCount() const override;
+	virtual TValueOrError<FEntityEventRecord, void> GetEntityEvent(uint64 EventIndex) const override;
+	virtual void EnumerateEntityEvents(
+		uint64 StartIndex,
+		uint64 Count,
+		TFunctionRef<void(const FEntityEventRecord&, uint64 /*EventIndex*/)> Callback) const override;
+
+	virtual uint64 GetRegionCount() const override;
+	virtual int32 GetLaneCount() const override;
+
+	virtual const FRegionLane* GetLane(int32 Index) const override;
+
+	virtual bool EnumerateRegions(double IntervalStart, double IntervalEnd, TFunctionRef<bool(const FRegion&)> Callback) const override;
+	virtual void EnumerateLanes(TFunctionRef<void(const FRegionLane&, const int32)> Callback) const override;
+
+	virtual uint64 GetUpdateCounter() const override
+	{
+		ReadAccessCheck();
+		return UpdateCounter;
+	}
+
+	//~ Edit operations
+
+	virtual void BeginEdit() const override
+	{
+		Lock.BeginWrite(GMassInsightsProviderLockState);
+	}
+	virtual void EndEdit() const override
+	{
+		Lock.EndWrite(GMassInsightsProviderLockState);
+	}
+	virtual void EditAccessCheck() const override
+	{
+		Lock.WriteAccessCheck(GMassInsightsProviderLockState);
+	}
+
+	virtual void AddFragment(const FFragmentInfo& FragmentInfo) override;
+	virtual void AddArchetype(const FArchetypeInfo& ArchetypeInfo) override;
+
+	virtual void BulkAddEntity(double Time, TConstArrayView<uint64> Entities, TConstArrayView<uint64> ArchetypeIDs) override;
+	virtual void BulkMoveEntity(double Time, TConstArrayView<uint64> Entities, TConstArrayView<uint64> ArchetypeIDs) override;
+	virtual void BulkDestroyEntity(double Time, TConstArrayView<uint64> Entities) override;
+
+	virtual void AppendRegionBegin(const TCHAR* Name, uint64 ID, double Time) override;
+	virtual void AppendRegionBegin(const TCHAR* Name, double Time) override;
+	virtual void AppendRegionEnd(const TCHAR* ID, double Time) override;
+	virtual void AppendRegionEnd(const uint64 ID, double Time) override;
+
+	virtual void OnAnalysisSessionEnded() override;
+
+	//////////////////////////////////////////////////
+
+private:
+	/** Update the depth member of a region to allow overlapping regions to be displayed on separate lanes. */
+	int32 CalculateRegionDepth(const FRegion& Item) const;
+	void AppendRegionEndInternal(FRegion* OpenRegion, double Time);
+
+private:
+	mutable TraceServices::FProviderLock Lock;
+
+	TraceServices::IAnalysisSession& Session;
+
+	TMap<uint64, const FFragmentInfo*> FragmentInfoByID;
+	TChunkedArray<FFragmentInfo> FragmentInfos;
+
+	/** Allocation of ranges of fragments, mostly for ArchetypeInfo */
+	TChunkedArray<const FFragmentInfo*> FragmentInfoRanges;
+
+	TMap<uint64, const FArchetypeInfo*> ArchetypeByID;
+	TChunkedArray<FArchetypeInfo> ArchetypeInfos;
+
+	/** Sorted by Cycle */
+	TraceServices::TPagedArray<FEntityEventRecord> EntityEvents;
+
+	/** Open regions inside lanes */
+	TMap<FStringView, FRegion*> OpenRegionsByName;
+	TMap<uint64, FRegion*> OpenRegionsByID;
+
+	/** Closed regions */
+	TArray<FRegionLane> Lanes;
+
+	/** Counter incremented each time region data changes during analysis */
+	uint64 UpdateCounter = 0;
+
+	static constexpr uint32 MaxWarningMessages = 100;
+	static constexpr uint32 MaxErrorMessages = 100;
+
+	uint32 NumWarnings = 0;
+	uint32 NumErrors = 0;
+};
+
+} // namespace UE::Mass::Trace
+
+#endif // UE_MASS_TRACE_ANALYSIS_ENABLED
diff --git a/Engine/Source/Runtime/Mass/MassDeveloper/Public/Trace/MassTraceTypes.h b/Engine/Source/Runtime/Mass/MassDeveloper/Public/Trace/MassTraceTypes.h
new file mode 100644
--- /dev/null
+++ b/Engine/Source/Runtime/Mass/MassDeveloper/Public/Trace/MassTraceTypes.h
@@ -0,0 +1,235 @@
+// Copyright Epic Games, Inc. All Rights Reserved.
+
+#pragma once
+
+#if UE_MASS_TRACE_ANALYSIS_ENABLED
+
+#include "Common/PagedArray.h" // TraceServices
+#include "TraceServices/Model/AnalysisSession.h"
+
+#define UE_API MASSDEVELOPER_API
+
+template <typename ValueType, typename ErrorType> class TValueOrError;
+
+namespace UE::Mass::Trace
+{
+
+struct FRegion
+{
+	double BeginTime = std::numeric_limits<double>::infinity();
+	double EndTime = std::numeric_limits<double>::infinity();
+	const TCHAR* Text = nullptr;
+	/** ID will be zero if the region is identified by Name only */
+	uint64 ID = 0;
+	int32 Depth = -1;
+};
+
+class FRegionLane
+{
+	friend class FMassTraceProvider;
+
+public:
+	explicit FRegionLane(TraceServices::ILinearAllocator& InAllocator) : Regions(InAllocator, 512)
+	{
+	}
+
+	int32 Num() const
+	{
+		return static_cast<int32>(Regions.Num());
+	}
+
+	/**
+	 * Call Callback for every region overlapping the interval defined by IntervalStart and IntervalEnd
+ 	 * @param IntervalStart the start time of the interval to enumerate regions for.
+	 * @param IntervalEnd the end time of the interval to enumerate regions for.
+	 * @param Callback a callback called for each region. Return false to abort iteration.
+	 * @returns true if the enumeration finished, false if it was aborted by the callback returning false
+	 */
+	UE_API bool EnumerateRegions(double IntervalStart, double IntervalEnd, TFunctionRef<bool(const FRegion&)> Callback) const;
+
+private:
+	TraceServices::TPagedArray<FRegion> Regions;
+};
+
+enum class EFragmentType : uint8
+{
+	Unknown,
+	Fragment,
+	Tag,
+	Shared,
+	Sparse,
+	ConstShared,
+	SparseTag
+};
+
+struct FFragmentInfo
+{
+	uint64 Id = INDEX_NONE;
+	FString Name;
+	uint32 Size = 0;
+	EFragmentType Type = EFragmentType::Unknown;
+
+	const TCHAR* GetName() const
+	{
+		return *Name;
+	}
+};
+
+struct FArchetypeInfo
+{
+	uint64 Id = INDEX_NONE;
+	TArray<const FFragmentInfo*> Fragments;
+
+	TConstArrayView<const FFragmentInfo*> GetFragments() const
+	{
+		return MakeArrayView(Fragments);
+	}
+};
+
+enum class EEntityEventType : uint8
+{
+	Unknown,
+	Created,
+	ArchetypeChange,
+	Destroyed
+};
+
+struct FEntityEventRecord
+{
+	double Time = 0;
+	uint64 Entity = INDEX_NONE;
+	uint64 ArchetypeID = INDEX_NONE;
+	EEntityEventType Operation = EEntityEventType::Unknown;
+};
+
+class IMassTraceProvider
+	: public TraceServices::IProvider
+{
+public:
+	virtual ~IMassTraceProvider() override = default;
+
+	virtual int32 GetFragmentCount() const = 0;
+	virtual const FFragmentInfo* FindFragmentById(uint64 FragmentId) const = 0;
+	virtual const FArchetypeInfo* FindArchetypeById(uint64 ArchetypeId) const = 0;
+
+	virtual void EnumerateFragments(TFunctionRef<void(const FFragmentInfo& FragmentInfo, int32 Index)> Callback, int32 BeginIndex = 0) const = 0;
+
+	virtual TValueOrError<FEntityEventRecord, void> GetEntityEvent(uint64 EventIndex) const = 0;
+
+	virtual uint64 GetEntityEventCount() const = 0;
+
+	/**
+	 * Enumerate up to Count number of events starting at the StartIndex
+	 * Enumeration will end early if there is not enough events or if Callback returns false
+	 */
+	virtual void EnumerateEntityEvents(
+		uint64 StartIndex,
+		uint64 Count,
+		TFunctionRef<void(const FEntityEventRecord&, uint64 EventIndex)> Callback) const = 0;
+
+	/**
+	 * @return the amount of currently known regions (including open-ended ones)
+	 */
+	virtual uint64 GetRegionCount() const = 0;
+
+	/**
+	 * @return the number of lanes
+	 */
+	virtual int32 GetLaneCount() const = 0;
+
+	/**
+	 * Direct access to a certain lane at a given index/depth.
+	 * The pointer is valid only in the current read scope.
+	 * @return a pointer to the lane at the specified depth index or nullptr if Index > GetLaneCount()-1
+	 */
+	virtual const FRegionLane* GetLane(int32 Index) const = 0;
+
+	/**
+	 * Enumerates all regions that overlap a certain time interval. Will enumerate by depth but does not expose lanes.
+	 * @param IntervalStart the start time of the interval to enumerate regions for.
+	 * @param IntervalEnd the end time of the interval to enumerate regions for.
+	 * @param Callback a callback called for each region. Return false to abort iteration.
+	 * @returns true if the enumeration finished, false if it was aborted by the callback returning false
+	 */
+	virtual bool EnumerateRegions(double IntervalStart, double IntervalEnd, TFunctionRef<bool(const FRegion&)> Callback) const = 0;
+
+	/**
+	 * Will call Callback(Lane, Depth) for each lane in order.
+	 */
+	virtual void EnumerateLanes(TFunctionRef<void(const FRegionLane&, const int32)> Callback) const = 0;
+
+	/**
+	 * @return A monotonically increasing counter that that changes each time new data is added to the provider.
+	 * This can be used to detect when to update any (UI-)state dependent on the provider during analysis.
+	 */
+	virtual uint64 GetUpdateCounter() const = 0;
+};
+
+/**
+ * The interface to a provider that can consume mutations of region events from a session.
+ */
+class IEditableMassTraceProvider
+	: public TraceServices::IEditableProvider
+{
+public:
+	virtual ~IEditableMassTraceProvider() override = default;
+
+	virtual void AddFragment(const FFragmentInfo& FragmentInfo) = 0;
+	virtual void AddArchetype(const FArchetypeInfo& ArchetypeInfo) = 0;
+
+	virtual void BulkAddEntity(double Time, TConstArrayView<uint64> Entities, TConstArrayView<uint64> ArchetypeIDs) = 0;
+	virtual void BulkMoveEntity(double Time, TConstArrayView<uint64> Entities, TConstArrayView<uint64> ArchetypeIDs) = 0;
+	virtual void BulkDestroyEntity(double Time, TConstArrayView<uint64> Entities) = 0;
+
+	/**
+	 * Append a new begin event of a region from the trace session.
+	 *
+	 * @param Name		The string name of the region.
+	 * @param Time		The time in seconds of the begin event of this region.
+	 */
+	virtual void AppendRegionBegin(const TCHAR* Name, double Time) = 0;
+
+	/**
+	 * Append a new begin event of a region from the trace session.
+	 *
+	 * @param Name		The string name of the region.
+	 * @param ID		The ID of the region. Used to uniquely identify regions with the same name.
+	 * @param Time		The time in seconds of the begin event of this region.
+	 */
+	virtual void AppendRegionBegin(const TCHAR* Name, uint64 ID, double Time) = 0;
+
+	/**
+	 * Append a new end event of a region from the trace session (by Name).
+	 *
+	 * @param Name		The string name of the region.
+	 * @param Time		The time in seconds of the end event of this region.
+	 */
+	virtual void AppendRegionEnd(const TCHAR* Name, double Time) = 0;
+
+	/**
+	 * Append a new end event of a region from the trace session (by ID).
+	 *
+	 * @param ID		The ID of the region.
+	 * @param Time		The time in seconds of the end event of this region.
+	 */
+	virtual void AppendRegionEnd(const uint64 ID, double Time) = 0;
+
+	/**
+	 * Called from the analyzer once all events have been processed.
+	 * Allows postprocessing and error reporting for regions that were never closed.
+	 */
+	virtual void OnAnalysisSessionEnded() = 0;
+};
+
+////////////////////////////////////////////////////////////////////////////////////////////////////
+
+UE_API FName GetProviderName();
+UE_API void SetupMassTraceAnalysis(TraceServices::IAnalysisSession& InSession);
+
+UE_API const IMassTraceProvider& ReadMassTraceProvider(const TraceServices::IAnalysisSession& Session);
+UE_API IEditableMassTraceProvider& EditMassTraceProvider(TraceServices::IAnalysisSession& Session);
+} // namespace UE::Mass::Trace
+
+#undef UE_API
+
+#endif // UE_MASS_TRACE_ANALYSIS_ENABLED
