[MassSerialization] Remove Transient from fragment properties
Remove Transient flag from Mass fragment properties that need to survive serialization.
Runtime-only fragments correctly retain their Transient flags.

#Synced-CL 52346140

diff --git a/Engine/Source/Runtime/Mass/MassCore/Public/Mass/EntityFragments.h b/Engine/Source/Runtime/Mass/MassCore/Public/Mass/EntityFragments.h
--- a/Engine/Source/Runtime/Mass/MassCore/Public/Mass/EntityFragments.h
+++ b/Engine/Source/Runtime/Mass/MassCore/Public/Mass/EntityFragments.h
@@ -33,7 +33,7 @@ struct FTransformFragment : public FMassFragment
 	}
 
 protected:
-	UPROPERTY(Transient)
+	UPROPERTY()
 	FTransform Transform;
 };
 
diff --git a/Engine/Source/Runtime/Mass/MassEngine/Public/Mesh/MassEngineMeshFragments.h b/Engine/Source/Runtime/Mass/MassEngine/Public/Mesh/MassEngineMeshFragments.h
--- a/Engine/Source/Runtime/Mass/MassEngine/Public/Mesh/MassEngineMeshFragments.h
+++ b/Engine/Source/Runtime/Mass/MassEngine/Public/Mesh/MassEngineMeshFragments.h
@@ -30,7 +30,7 @@ struct FMassStaticMeshFragment : public FMassConstSharedFragment
 	FMassStaticMeshFragment() = default;
 	UE_API explicit FMassStaticMeshFragment(TNotNull<const UStaticMesh*> InMesh);
 
-	UPROPERTY(Transient, VisibleAnywhere, Category=Debug)
+	UPROPERTY(VisibleAnywhere, Category=Debug)
 	TWeakObjectPtr<const UStaticMesh> Mesh;
 };
 
@@ -42,33 +42,33 @@ struct FMassVisualizationMeshFragment : public FMassConstSharedFragment
 {
 	GENERATED_BODY()
 
-	UPROPERTY(Transient, VisibleAnywhere, Category = Mesh)
+	UPROPERTY(VisibleAnywhere, Category = Mesh)
 	int64 CustomDepthStencilValue = 0;
-	UPROPERTY(Transient, VisibleAnywhere, Category = Mesh)
+	UPROPERTY(VisibleAnywhere, Category = Mesh)
 	int32 TranslucencySortPriority = 0;
-	UPROPERTY(Transient, VisibleAnywhere, Category = Mesh)
+	UPROPERTY(VisibleAnywhere, Category = Mesh)
 	uint8 HiddenInGame : 1 = false;
-	UPROPERTY(Transient, VisibleAnywhere, Category = Mesh)
+	UPROPERTY(VisibleAnywhere, Category = Mesh)
 	uint8 CastShadow : 1 = true;
-	UPROPERTY(Transient, VisibleAnywhere, Category = Mesh)
+	UPROPERTY(VisibleAnywhere, Category = Mesh)
 	uint8 CastShadowAsTwoSided : 1 = false;
-	UPROPERTY(Transient, VisibleAnywhere, Category = Mesh)
+	UPROPERTY(VisibleAnywhere, Category = Mesh)
 	uint8 CastHiddenShadow : 1 = false;
-	UPROPERTY(Transient, VisibleAnywhere, Category = Mesh)
+	UPROPERTY(VisibleAnywhere, Category = Mesh)
 	uint8 AffectDynamicIndirectLighting : 1 = true;
-	UPROPERTY(Transient, VisibleAnywhere, Category = Mesh)
+	UPROPERTY(VisibleAnywhere, Category = Mesh)
 	uint8 AffectIndirectLightingWhileHidden : 1 = false;
-	UPROPERTY(Transient, VisibleAnywhere, Category = Mesh)
+	UPROPERTY(VisibleAnywhere, Category = Mesh)
 	uint8 ReceivesDecals : 1 = true;
-	UPROPERTY(Transient, VisibleAnywhere, Category = Mesh)
+	UPROPERTY(VisibleAnywhere, Category = Mesh)
 	uint8 NeverDistanceCull : 1 = false;
-	UPROPERTY(Transient, VisibleAnywhere, Category = Mesh)
+	UPROPERTY(VisibleAnywhere, Category = Mesh)
 	uint8 CastFarShadow : 1 = false;
-	UPROPERTY(Transient, VisibleAnywhere, Category = Mesh)
+	UPROPERTY(VisibleAnywhere, Category = Mesh)
 	uint8 CastInsetShadow : 1 = false;
-	UPROPERTY(Transient, VisibleAnywhere, Category = Mesh)
+	UPROPERTY(VisibleAnywhere, Category = Mesh)
 	uint8 CastContactShadow : 1 = false;
-	UPROPERTY(Transient, VisibleAnywhere, Category = Mesh)
+	UPROPERTY(VisibleAnywhere, Category = Mesh)
 	uint8 RenderCustomDepth : 1 = false;
 };
 
@@ -81,11 +81,11 @@ struct FMassEditorVisualizationMeshFragment : public FMassConstSharedFragment
 	GENERATED_BODY()
 
 #if  WITH_EDITORONLY_DATA
-	UPROPERTY(Transient, VisibleAnywhere, Category = Mesh)
+	UPROPERTY(VisibleAnywhere, Category = Mesh)
 	uint8 bShouldRenderSelected : 1 = false;
-	UPROPERTY(Transient, VisibleAnywhere, Category = Mesh)
+	UPROPERTY(VisibleAnywhere, Category = Mesh)
 	uint8 bHiddenInEditor : 1 = false;
-	UPROPERTY(Transient, VisibleAnywhere, Category = Mesh)
+	UPROPERTY(VisibleAnywhere, Category = Mesh)
 	uint8 bLevelInstanceEditingState : 1 = false;
 #endif //  WITH_EDITORONLY_DATA
 };
@@ -98,7 +98,7 @@ struct FMassOverrideMaterialsFragment : public FMassConstSharedFragment
 {
 	GENERATED_BODY()
 
-	UPROPERTY(Transient, VisibleAnywhere, Category = Debug)
+	UPROPERTY(VisibleAnywhere, Category = Debug)
 	TArray<TObjectPtr<UMaterialInterface>> OverrideMaterials;
 };
 
@@ -325,7 +325,7 @@ struct FMassISMGroupingFragment : public FMassConstSharedFragment
 		: GroupId(InGroupId)
 	{}
 
-	UPROPERTY(Transient, VisibleAnywhere, Category = Debug)
+	UPROPERTY(VisibleAnywhere, Category = Debug)
 	FGuid GroupId;
 };
 
diff --git a/Engine/Source/Runtime/Mass/MassEngine/Public/Physics/MassEnginePhysicsFragments.h b/Engine/Source/Runtime/Mass/MassEngine/Public/Physics/MassEnginePhysicsFragments.h
--- a/Engine/Source/Runtime/Mass/MassEngine/Public/Physics/MassEnginePhysicsFragments.h
+++ b/Engine/Source/Runtime/Mass/MassEngine/Public/Physics/MassEnginePhysicsFragments.h
@@ -22,10 +22,10 @@ struct FMassColliderFragment : public FMassConstSharedFragment
 	FMassColliderFragment() = default;
 	UE_API explicit FMassColliderFragment(TNotNull<const UStaticMesh*> Mesh);
 
-	UPROPERTY(Transient, VisibleAnywhere, Category=Debug)
+	UPROPERTY(VisibleAnywhere, Category=Debug)
 	TEnumAsByte<ECollisionEnabled::Type> CollisionType = ECollisionEnabled::NoCollision;
 
-	UPROPERTY(Transient, VisibleAnywhere, Category=Debug)
+	UPROPERTY(VisibleAnywhere, Category=Debug)
 	FCollisionResponseContainer CollisionResponse;
 };
 
