[Mass] Add lighting channel support for mesh fragments

#Synced-CL 52482515

diff --git a/Engine/Source/Runtime/Mass/MassEngine/Private/Mesh/MassEngineMeshUtils.cpp b/Engine/Source/Runtime/Mass/MassEngine/Private/Mesh/MassEngineMeshUtils.cpp
--- a/Engine/Source/Runtime/Mass/MassEngine/Private/Mesh/MassEngineMeshUtils.cpp
+++ b/Engine/Source/Runtime/Mass/MassEngine/Private/Mesh/MassEngineMeshUtils.cpp
@@ -133,6 +133,7 @@ namespace UE::MassEngine::Mesh
 		StaticMeshSceneProxyDesc.bReceivesDecals = MeshFragment.ReceivesDecals;
 		StaticMeshSceneProxyDesc.bRenderCustomDepth = MeshFragment.RenderCustomDepth;
 		StaticMeshSceneProxyDesc.CustomDepthStencilValue = (int32)MeshFragment.CustomDepthStencilValue;
+		StaticMeshSceneProxyDesc.LightingChannels = MeshFragment.LightingChannels;
 
 		if (MeshFragment.NeverDistanceCull)
 		{
diff --git a/Engine/Source/Runtime/Mass/MassEngine/Public/Mesh/MassEngineMeshFragments.h b/Engine/Source/Runtime/Mass/MassEngine/Public/Mesh/MassEngineMeshFragments.h
--- a/Engine/Source/Runtime/Mass/MassEngine/Public/Mesh/MassEngineMeshFragments.h
+++ b/Engine/Source/Runtime/Mass/MassEngine/Public/Mesh/MassEngineMeshFragments.h
@@ -47,6 +47,8 @@ struct FMassVisualizationMeshFragment : public FMassConstSharedFragment
 	UPROPERTY(VisibleAnywhere, Category = Mesh)
 	int32 TranslucencySortPriority = 0;
 	UPROPERTY(VisibleAnywhere, Category = Mesh)
+	FLightingChannels LightingChannels;
+	UPROPERTY(VisibleAnywhere, Category = Mesh)
 	uint8 HiddenInGame : 1 = false;
 	UPROPERTY(VisibleAnywhere, Category = Mesh)
 	uint8 CastShadow : 1 = true;
