[MassEngine][RigidPhysicsCore] Batched BodyInstance initialization via Mass processors for SceneGraph entities

Added MassPhysics processors to batch-initialize FBodyInstances through the Mass pipeline, replacing per-entity immediate initialization for SceneGraph entities.

New Mass processors:
- UMassPhysicsBodyInstancesBatcher: groups entities by mesh/collision/body-type into batch requests
- UMassPhysicsBodyInstancesBatchInitializer: calls FInitBodiesHelper on batched requests
- UMassPhysicsBodyInstanceReInitializer: observer that re-queues initialization on config change

New fragments/tags:
- FMassPhysicsCollisionSettingsFragment (replaces FMassColliderFragment)
- FMassPhysicsBodyInstanceFragment, FMassPhysicsBodyInstancesInitializationRequestFragment
- FMassPhysicsDynamicBodyTag, FMassPhysicsSimulatedBodyTag, FMassPhysicsBodyInstanceInitializedTag

FVerseRigidBody changes:
- BodyInstance and ChaosUserData converted from value to TSharedPtr for shared ownership with Mass fragments
- Added "strictly static" concept based on mesh part availability (no mesh parts = strictly static) and use the optimized FInitBodiesHelper</*bCompileStatic*/> initialization path
- Collision settings management centralized in VerseRigidBody (UpdateCollisionSettingsFragment)

Gated with CVar: "SceneGraph.EnableBodyInstancesInitializationInMass" (default: disabled)

Body type classification:
- "strictly static": no mesh parts available
- "non-simulated dynamic": has mesh parts, dynamic flag off
- "simulated dynamic": has mesh parts, dynamic flag on

Runtime transitions supported:
- Toggling the dynamic flag transitions between non-simulated and simulated dynamic
- Gaining mesh parts at runtime transitions from strictly static to dynamic
- Losing mesh parts transitions from dynamic back to strictly static

#rnx
#Synced-CL 52538360

diff --git a/Engine/Source/Runtime/Mass/MassEngine/MassEngine.Build.cs b/Engine/Source/Runtime/Mass/MassEngine/MassEngine.Build.cs
--- a/Engine/Source/Runtime/Mass/MassEngine/MassEngine.Build.cs
+++ b/Engine/Source/Runtime/Mass/MassEngine/MassEngine.Build.cs
@@ -14,6 +14,7 @@ namespace UnrealBuildTool.Rules
 
 			PublicDependencyModuleNames.AddRange(
 				new string[] {
+					"Chaos",
 					"Core",
 					"CoreUObject",
 					"Engine",
diff --git a/Engine/Source/Runtime/Mass/MassEngine/Private/AI/MassEngineNavigationProcessors.cpp b/Engine/Source/Runtime/Mass/MassEngine/Private/AI/MassEngineNavigationProcessors.cpp
--- a/Engine/Source/Runtime/Mass/MassEngine/Private/AI/MassEngineNavigationProcessors.cpp
+++ b/Engine/Source/Runtime/Mass/MassEngine/Private/AI/MassEngineNavigationProcessors.cpp
@@ -22,7 +22,7 @@
 namespace UE::NavigationRelevant
 {
 
-bool IsNavigationRelevant(const FMassColliderFragment& ColliderFragment,
+bool IsNavigationRelevant(const FMassPhysicsCollisionSettingsFragment& CollisionSettingsFragment,
 	const FMassStaticMeshFragment& MeshFragment,
 	const FTransformFragment& TransformFragment)
 {
@@ -33,12 +33,12 @@ bool IsNavigationRelevant(const FMassColliderFragment& ColliderFragment,
 
 	// @see UPrimitiveComponent::IsNavigationRelevant()
 	{
-		if (!CollisionEnabledHasQuery(ColliderFragment.CollisionType))
+		if (!CollisionEnabledHasQuery(CollisionSettingsFragment.CollisionType))
 		{
 			return false;
 		}
 
-		const FCollisionResponseContainer& ResponseToChannels = ColliderFragment.CollisionResponse;
+		const FCollisionResponseContainer& ResponseToChannels = CollisionSettingsFragment.CollisionResponse;
 		if (!(ResponseToChannels.GetResponse(ECC_Pawn) == ECR_Block
 			|| ResponseToChannels.GetResponse(ECC_Vehicle) == ECR_Block))
 		{
@@ -157,7 +157,7 @@ UMassNavigationElementProcessor::UMassNavigationElementProcessor(const FObjectIn
 
 void UMassNavigationElementProcessor::ConfigureQueries(const TSharedRef<FMassEntityManager>& EntityManager)
 {
-	NewRelevantEntityQuery.AddConstSharedRequirement<FMassColliderFragment>();
+	NewRelevantEntityQuery.AddConstSharedRequirement<FMassPhysicsCollisionSettingsFragment>();
 	NewRelevantEntityQuery.AddConstSharedRequirement<FMassStaticMeshFragment>();
 	NewRelevantEntityQuery.AddRequirement<FTransformFragment>(EMassFragmentAccess::ReadOnly);
 	NewRelevantEntityQuery.AddRequirement<FMassNavigationRelevantFragment>(EMassFragmentAccess::ReadOnly, EMassFragmentPresence::None);
@@ -180,20 +180,21 @@ void UMassNavigationElementProcessor::Execute(FMassEntityManager& EntityManager,
 		{
 			const TConstArrayView<FTransformFragment> TransformFragments = Context.GetFragmentView<FTransformFragment>();
 			const FMassStaticMeshFragment& MeshFragment = Context.GetConstSharedFragment<FMassStaticMeshFragment>();
-			const FMassColliderFragment& ColliderFragment = Context.GetConstSharedFragment<FMassColliderFragment>();
+			const FMassPhysicsCollisionSettingsFragment& CollisionSettingsFragment = Context.GetConstSharedFragment<FMassPhysicsCollisionSettingsFragment>();
 
 			for (FMassExecutionContext::FEntityIterator EntityIt = Context.CreateEntityIterator(); EntityIt; ++EntityIt)
 			{
 				const FMassEntityHandle EntityHandle = Context.GetEntity(EntityIt);
 				FNavigationElementHandle NavigationElementHandle = FNavigationElementHandle::Invalid;
 
-				if (IsNavigationRelevant(ColliderFragment, MeshFragment, TransformFragments[EntityIt]))
+				if (!IsNavigationRelevant(CollisionSettingsFragment, MeshFragment, TransformFragments[EntityIt]))
 				{
-					FNavigationElement Element(*EntityManagerOwner, EntityHandle.AsNumber());
-					ConfigureElement(EntityManagerOwner, Element, MeshFragment, TransformFragments[EntityIt], ParametersSharedFragment.Get<const FMassNavigationRelevantParameters>());
-					NavigationElementHandle = FNavigationSystem::AddNavigationElement(EntityManagerOwner->GetWorld(), MoveTemp(Element));
+					continue;
 				}
 
+				FNavigationElement Element(*EntityManagerOwner, EntityHandle.AsNumber());
+				ConfigureElement(EntityManagerOwner, Element, MeshFragment, TransformFragments[EntityIt], ParametersSharedFragment.Get<const FMassNavigationRelevantParameters>());
+				NavigationElementHandle = FNavigationSystem::AddNavigationElement(EntityManagerOwner->GetWorld(), MoveTemp(Element));
 				Context.Defer().PushCommand<FMassCommandAddFragmentInstances>(EntityHandle, FMassNavigationRelevantFragment(NavigationElementHandle));
 				Context.Defer().PushCommand<FMassDeferredAddCommand>([EntityHandle, ParametersSharedFragment](FMassEntityManager& InEntityManager)
 					{
@@ -236,7 +237,7 @@ void UMassDirtyNavigationRelevantUpdateProcessor::InitializeInternal(UObject& Ow
 
 void UMassDirtyNavigationRelevantUpdateProcessor::ConfigureQueries(const TSharedRef<FMassEntityManager>& EntityManager)
 {
-	EntityQuery.AddConstSharedRequirement<FMassColliderFragment>();
+	EntityQuery.AddConstSharedRequirement<FMassPhysicsCollisionSettingsFragment>();
 	EntityQuery.AddConstSharedRequirement<FMassStaticMeshFragment>();
 	EntityQuery.AddConstSharedRequirement<FMassNavigationRelevantParameters>();
 	EntityQuery.AddRequirement<FTransformFragment>(EMassFragmentAccess::ReadOnly);
@@ -259,7 +260,7 @@ void UMassDirtyNavigationRelevantUpdateProcessor::SignalEntities(
 			const TConstArrayView<FTransformFragment> TransformFragments = Context.GetFragmentView<FTransformFragment>();
 			const TArrayView<FMassNavigationRelevantFragment> NavigationRelevantFragments = Context.GetMutableFragmentView<FMassNavigationRelevantFragment>();
 
-			const FMassColliderFragment& ColliderFragment = Context.GetConstSharedFragment<FMassColliderFragment>();
+			const FMassPhysicsCollisionSettingsFragment& CollisionSettingsFragment = Context.GetConstSharedFragment<FMassPhysicsCollisionSettingsFragment>();
 			const FMassStaticMeshFragment& MeshFragment = Context.GetConstSharedFragment<FMassStaticMeshFragment>();
 			const FMassNavigationRelevantParameters& NavigationRelevantParameters = Context.GetConstSharedFragment<FMassNavigationRelevantParameters>();
 
@@ -267,7 +268,7 @@ void UMassDirtyNavigationRelevantUpdateProcessor::SignalEntities(
 			{
 				const bool bWasRelevant = NavigationRelevantFragments[EntityIt].Handle.IsValid();
 
-				if (IsNavigationRelevant(ColliderFragment, MeshFragment, TransformFragments[EntityIt]))
+				if (IsNavigationRelevant(CollisionSettingsFragment, MeshFragment, TransformFragments[EntityIt]))
 				{
 					FNavigationElement Element(*EntityManagerOwner, Context.GetEntity(EntityIt).AsNumber());
 					ConfigureElement(EntityManagerOwner, Element, MeshFragment, TransformFragments[EntityIt], NavigationRelevantParameters);
@@ -296,13 +297,19 @@ UMassNavigationRelevantMassDeinitializer::UMassNavigationRelevantMassDeinitializ
 	: EntityQuery(*this)
 {
 	ExecutionFlags = static_cast<int32>(EProcessorExecutionFlags::AllWorldModes);
-	ObservedTypes.Add(FMassNavigationRelevantFragment::StaticStruct());
+	ObservedTypes.Append({
+		FMassStaticMeshFragment::StaticStruct(),
+		FTransformFragment::StaticStruct(),
+		FMassPhysicsCollisionSettingsFragment::StaticStruct()});
 	ObservedOperations = EMassObservedOperationFlags::Remove;
 	bRequiresGameThreadExecution = true;
 }
 
 void UMassNavigationRelevantMassDeinitializer::ConfigureQueries(const TSharedRef<FMassEntityManager>& EntityManager)
 {
+	EntityQuery.AddConstSharedRequirement<FMassStaticMeshFragment>(EMassFragmentPresence::All);
+	EntityQuery.AddConstSharedRequirement<FMassPhysicsCollisionSettingsFragment>(EMassFragmentPresence::All);
+	EntityQuery.AddRequirement<FTransformFragment>(EMassFragmentAccess::ReadOnly);
 	EntityQuery.AddRequirement<FMassNavigationRelevantFragment>(EMassFragmentAccess::ReadOnly);
 }
 
@@ -319,5 +326,7 @@ void UMassNavigationRelevantMassDeinitializer::Execute(FMassEntityManager& Entit
 					FNavigationSystem::RemoveNavigationElement(World, NavigationRelevantFragments[EntityIt].Handle);
 				}
 			}
+
+			Context.Defer().PushCommand<FMassCommandRemoveElements<FMassNavigationRelevantFragment, FMassNavigationRelevantParameters>>(Context.GetEntities());
 		});
 }
diff --git a/Engine/Source/Runtime/Mass/MassEngine/Private/Mesh/MassStaticMeshRenderStateHelper.cpp b/Engine/Source/Runtime/Mass/MassEngine/Private/Mesh/MassStaticMeshRenderStateHelper.cpp
--- a/Engine/Source/Runtime/Mass/MassEngine/Private/Mesh/MassStaticMeshRenderStateHelper.cpp
+++ b/Engine/Source/Runtime/Mass/MassEngine/Private/Mesh/MassStaticMeshRenderStateHelper.cpp
@@ -54,9 +54,9 @@ UMaterialInterface* FMassBaseStaticMeshRenderStateHelper::GetMaterial(int32 Mate
 
 const FCollisionResponseContainer& FMassBaseStaticMeshRenderStateHelper::GetCollisionResponseToChannels() const
 {
-	if (FMassColliderFragment* ColliderFragment = GetEntityManager().GetConstSharedFragmentDataPtr<FMassColliderFragment>(EntityHandle))
+	if (FMassPhysicsCollisionSettingsFragment* CollisionSettingsFragment = GetEntityManager().GetConstSharedFragmentDataPtr<FMassPhysicsCollisionSettingsFragment>(EntityHandle))
 	{
-		return ColliderFragment->CollisionResponse;
+		return CollisionSettingsFragment->CollisionResponse;
 	}
 
 	return FCollisionResponseContainer::GetDefaultResponseContainer();
diff --git a/Engine/Source/Runtime/Mass/MassEngine/Private/RigidBody/MassPhysicsProcessors.cpp b/Engine/Source/Runtime/Mass/MassEngine/Private/RigidBody/MassPhysicsProcessors.cpp
new file mode 100644
--- /dev/null
+++ b/Engine/Source/Runtime/Mass/MassEngine/Private/RigidBody/MassPhysicsProcessors.cpp
@@ -0,0 +1,475 @@
+// Copyright Epic Games, Inc. All Rights Reserved.
+
+#include "MassPhysicsProcessors.h"
+
+#include "Chaos/ChaosUserEntity.h"
+#include "Mass/EntityFragments.h"
+#include "MassEngineTypes.h"
+#include "MassExecutionContext.h"
+#include "Mesh/MassEngineMeshFragments.h"
+#include "Misc/HashBuilder.h"
+#include "Physics/Experimental/PhysScene_Chaos.h"
+#include "Physics/MassEnginePhysicsFragments.h"
+#include "PhysicsEngine/BodyInstance.h"
+#include "PhysicsEngine/BodySetup.h"
+#include "PhysicsEngine/PhysicsObjectExternalInterface.h"
+#include "PhysicsProxy/SingleParticlePhysicsProxy.h"
+#include "VisualLogger/VisualLogger.h"
+
+#include UE_INLINE_GENERATED_CPP_BY_NAME(MassPhysicsProcessors)
+
+namespace UE::MassEngine::Private
+{
+enum class EBodyType : uint8
+{
+	Static,
+	NonSimulatedDynamic,
+	SimulatedDynamic
+};
+}
+
+//----------------------------------------------------------------------//
+// UMassPhysicsBodyInstancesBatcher
+//----------------------------------------------------------------------//
+UMassPhysicsBodyInstancesBatcher::UMassPhysicsBodyInstancesBatcher(const FObjectInitializer& ObjectInitializer)
+	: Super(ObjectInitializer)
+	, EntityQuery(*this)
+{
+	bRequiresGameThreadExecution = false;
+
+	ExecutionFlags = static_cast<int32>(EProcessorExecutionFlags::AllWorldModes);
+	ExecutionOrder.ExecuteInGroup = UE::Mass::ProcessorGroupNames::SetupBatchedInitializations;
+}
+
+void UMassPhysicsBodyInstancesBatcher::ConfigureQueries(const TSharedRef<FMassEntityManager>& EntityManager)
+{
+	EntityQuery.AddConstSharedRequirement<FMassStaticMeshFragment>();
+	EntityQuery.AddConstSharedRequirement<FMassPhysicsCollisionSettingsFragment>();
+
+	EntityQuery.AddRequirement<FTransformFragment>(EMassFragmentAccess::ReadOnly);
+	EntityQuery.AddRequirement<FMassPhysicsBodyInstanceFragment>(EMassFragmentAccess::ReadWrite);
+
+	// Init
+	EntityQuery.AddTagRequirement<FMassPhysicsBodyInstanceInitializedTag>(EMassFragmentPresence::None);
+
+	// Create the batching entity and link it but let the UMassPhysicsBodyInstancesBatchInitializer initialize the instances
+	FMassElementBitSet Bitset;
+	Bitset.Add<FMassPhysicsBodyInstancesInitializationRequestFragment>();
+	BatchInitializationRequestArchetype = EntityManager->CreateArchetype(Bitset, FMassArchetypeCreationParams{ "PhysicsBodyBatchInitializationRequest" });
+	ProcessorEntityCreationRequirements.AddCreatedArchetype(Bitset);
+}
+
+void UMassPhysicsBodyInstancesBatcher::Execute(FMassEntityManager& EntityManager, FMassExecutionContext& ExecutionContext)
+{
+	using namespace UE::MassEngine::Private;
+
+	struct FBatchInfo
+	{
+		const FMassStaticMeshFragment& StaticMeshFragment;
+		uint32 SourceSharedFragmentHash = 0;
+		uint32 InstanceCount = 0;
+		uint32 CurrentInstanceIndex = 0;
+		EBodyType BodyType = EBodyType::Static;
+	};
+
+	struct FBatches
+	{
+		const FMassPhysicsCollisionSettingsFragment* CollisionSettingsFragment = nullptr;
+		TArray<FBatchInfo> BatchInfos;
+		TArray<FMassEntityHandle> CreatedRequestEntities;
+	};
+
+	TArray<FBatches> Batches;
+
+	auto RetrieveBatchInfo =
+		[&Batches](const FMassStaticMeshFragment& StaticMeshFragment
+			, const FMassPhysicsCollisionSettingsFragment* CollisionSettingsFragment
+			, const EBodyType BodyType
+			, const bool bCreateIfNotFound
+			, FMassEntityHandle* OutRequestEntity = nullptr) -> FBatchInfo*
+		{
+			FBatches* Batch = Batches.FindByPredicate([CollisionSettingsFragment](const FBatches& Batch) { return Batch.CollisionSettingsFragment == CollisionSettingsFragment; });
+			if (Batch == nullptr)
+			{
+				if (!bCreateIfNotFound)
+				{
+					return nullptr;
+				}
+				Batch = &Batches.Emplace_GetRef(CollisionSettingsFragment);
+			}
+
+			// Build the source shared fragment hash
+			FHashBuilder HashBuilder;
+			auto AddFragmentToHash = [&HashBuilder](const void* FragmentPtr)
+				{
+					HashBuilder.AppendRaw(&FragmentPtr, sizeof(&FragmentPtr));
+				};
+			AddFragmentToHash(&StaticMeshFragment);
+			HashBuilder.Append(BodyType);
+
+			// Find if the batch info already exists, if not add it
+			FBatchInfo* BatchInfo;
+			int32 Index = Batch->BatchInfos.IndexOfByPredicate([Hash = HashBuilder.GetHash()](const FBatchInfo& EntityInfo) { return EntityInfo.SourceSharedFragmentHash == Hash; });
+			if (Index == INDEX_NONE)
+			{
+				if (!bCreateIfNotFound)
+				{
+					return nullptr;
+				}
+				Index = Batch->BatchInfos.Num();
+				BatchInfo = &Batch->BatchInfos.Emplace_GetRef(FBatchInfo{.StaticMeshFragment = StaticMeshFragment, .SourceSharedFragmentHash = HashBuilder.GetHash()});
+				BatchInfo->BodyType = BodyType;
+			}
+			else
+			{
+				BatchInfo = &Batch->BatchInfos[Index];
+			}
+
+			if (OutRequestEntity)
+			{
+				checkf(Index < Batch->CreatedRequestEntities.Num(), TEXT("Expecting all the entities to be created at this point"));
+				*OutRequestEntity = Batch->CreatedRequestEntities[Index];
+			}
+			return BatchInfo;
+		};
+
+	// First, gather all the different request entities we need to create.
+	EntityQuery.ForEachEntityChunk(ExecutionContext,
+		[&RetrieveBatchInfo](const FMassExecutionContext& Context)
+		{
+			const FMassStaticMeshFragment& StaticMeshFragment = Context.GetConstSharedFragment<FMassStaticMeshFragment>();
+			const FMassPhysicsCollisionSettingsFragment* ColliderFragment = Context.GetConstSharedFragmentPtr<FMassPhysicsCollisionSettingsFragment>();
+
+			const EBodyType BodyType =
+				Context.DoesArchetypeHaveTag<FMassPhysicsSimulatedBodyTag>() ?
+				EBodyType::SimulatedDynamic
+				: Context.DoesArchetypeHaveTag<FMassPhysicsDynamicBodyTag>() ?
+				EBodyType::NonSimulatedDynamic
+				: EBodyType::Static;
+
+			// Find the batch matching the shared fragment
+			FBatchInfo* BatchInfo = RetrieveBatchInfo(StaticMeshFragment, ColliderFragment, BodyType, /*bCreateIfNotFound*/true);
+			checkf(BatchInfo, TEXT("Expecting the batch to be created"));
+			BatchInfo->InstanceCount += Context.GetNumEntities();
+		});
+
+	// Batch create the mass entities for all batched initialization requests
+	for (FBatches& Batch : Batches)
+	{
+		ExecutionContext.BatchCreateEntities(BatchInitializationRequestArchetype, Batch.BatchInfos.Num(), Batch.CreatedRequestEntities);
+	}
+
+	// Gather transforms for each entity and linked them to the created request entity
+	EntityQuery.ForEachEntityChunk(ExecutionContext,
+		[this, &RetrieveBatchInfo](FMassExecutionContext& Context)
+		{
+			const FMassStaticMeshFragment& StaticMeshFragment = Context.GetConstSharedFragment<FMassStaticMeshFragment>();
+			const FMassPhysicsCollisionSettingsFragment* ColliderFragment = Context.GetConstSharedFragmentPtr<FMassPhysicsCollisionSettingsFragment>();
+			const TConstArrayView<FTransformFragment> TransformFragments = Context.GetFragmentView<FTransformFragment>();
+			const TArrayView<FMassPhysicsBodyInstanceFragment> BodyInstanceFragments = Context.GetMutableFragmentView<FMassPhysicsBodyInstanceFragment>();
+
+			const bool bSimulated = Context.DoesArchetypeHaveTag<FMassPhysicsSimulatedBodyTag>();
+			const bool bDynamic = Context.DoesArchetypeHaveTag<FMassPhysicsDynamicBodyTag>();
+
+			const EBodyType BodyType =
+				bSimulated ? EBodyType::SimulatedDynamic
+				: bDynamic ? EBodyType::NonSimulatedDynamic
+				: EBodyType::Static;
+
+			FMassEntityHandle RequestEntityHandle;
+			FBatchInfo* BatchInfo = RetrieveBatchInfo(StaticMeshFragment, ColliderFragment, BodyType, /*bCreateIfNotFound*/false, &RequestEntityHandle);
+			checkf(BatchInfo, TEXT("Expecting the batch to be created"));
+			checkf(BatchInfo->BodyType == BodyType, TEXT("Expecting the batch to be created"));
+			uint32& InstanceIndex = BatchInfo->CurrentInstanceIndex;
+
+			const FMassEntityManager& EntityManager = Context.GetEntityManagerChecked();
+
+			FMassPhysicsBodyInstancesInitializationRequestFragment& RequestFragment = EntityManager.GetFragmentDataChecked<FMassPhysicsBodyInstancesInitializationRequestFragment>(RequestEntityHandle);
+			RequestFragment.StaticMesh = BatchInfo->StaticMeshFragment.Mesh;
+			RequestFragment.Transforms.Reserve(BatchInfo->InstanceCount);
+			RequestFragment.InstanceBodies.Reserve(BatchInfo->InstanceCount);
+			RequestFragment.ChaosUserDefinedEntity.Reserve(BatchInfo->InstanceCount);
+			RequestFragment.bDynamic = BodyType != EBodyType::Static;
+
+			UE_VLOG_UELOG(this, LogMassEngine, Log, TEXT("Created initialization request for %d %s body instances using mesh: %s")
+				, BatchInfo->InstanceCount
+				, bSimulated ? TEXT("simulated dynamic") : bDynamic ? TEXT("non-simulated dynamic") : TEXT("static")
+				, *GetPathNameSafe(BatchInfo->StaticMeshFragment.Mesh.Get()));
+
+			for (FMassExecutionContext::FEntityIterator EntityIt = Context.CreateEntityIterator(); EntityIt; ++EntityIt)
+			{
+				ensureMsgf(!BodyInstanceFragments[*EntityIt].BodyInstanceOwnerObject.IsExplicitlyNull(), TEXT("BodyInstanceOwnerObject is required to determine the validity of the raw pointers"));
+				if (BodyInstanceFragments[*EntityIt].BodyInstanceOwnerObject.IsValid())
+				{
+					verifyf(RequestFragment.Transforms.Add(TransformFragments[*EntityIt].GetTransform()) == InstanceIndex
+						, TEXT("Expecting the added index to match the current instance"));
+
+					// Fragment created from serialized data need to create the body
+					if (!BodyInstanceFragments[*EntityIt].BodyInstance.IsValid())
+					{
+						BodyInstanceFragments[*EntityIt].BodyInstance = MakeShared<FBodyInstance>();
+
+						// @todo: ChaosUserDefinedEntity is currently not required for "mass-only" version
+					}
+					verifyf(RequestFragment.InstanceBodies.Add(BodyInstanceFragments[*EntityIt].BodyInstance) == InstanceIndex
+						, TEXT("Expecting the added index to match the current instance"));
+
+					verifyf(RequestFragment.ChaosUserDefinedEntity.Add(BodyInstanceFragments[*EntityIt].ChaosUserDefinedEntity) == InstanceIndex
+						, TEXT("Expecting the added index to match the current instance"));
+
+					ensure(bSimulated == RequestFragment.InstanceBodies[InstanceIndex]->bSimulatePhysics);
+					InstanceIndex++;
+				}
+				else
+				{
+					UE_LOG(LogMassEngine, Warning, TEXT("Skipped instance: BodyInstanceOwnerObject is no longer valid. Removing stale fragment."));
+					Context.Defer().RemoveFragment<FMassPhysicsBodyInstanceFragment>(Context.GetEntity(*EntityIt));
+				}
+
+				// Tag unconditionally to exclude from future Batcher runs.
+				// For valid entities this marks a successful batch; for stale-owner
+				// entities this prevents an infinite retry (the owner is destroyed and IsValid() will never return true again).
+				Context.Defer().AddTag<FMassPhysicsBodyInstanceInitializedTag>(Context.GetEntity(*EntityIt));
+			}
+		});
+}
+
+//----------------------------------------------------------------------//
+// UMassPhysicsBodyInstancesBatchInitializer
+//----------------------------------------------------------------------//
+UMassPhysicsBodyInstancesBatchInitializer::UMassPhysicsBodyInstancesBatchInitializer(const FObjectInitializer& ObjectInitializer)
+	: Super(ObjectInitializer)
+	, EntityQuery(*this)
+{
+	bRequiresGameThreadExecution = !FPhysScene::SupportsAsyncPhysicsStateCreation();
+
+	ExecutionFlags = static_cast<int32>(EProcessorExecutionFlags::AllWorldModes);
+
+	ExecutionOrder.ExecuteInGroup = UE::Mass::ProcessorGroupNames::ProcessBatchedInitializations;
+	ExecutionOrder.ExecuteAfter.Add(UE::Mass::ProcessorGroupNames::SetupBatchedInitializations);
+}
+
+void UMassPhysicsBodyInstancesBatchInitializer::ConfigureQueries(const TSharedRef<FMassEntityManager>& EntityManager)
+{
+	EntityQuery.AddRequirement<FMassPhysicsBodyInstancesInitializationRequestFragment>(EMassFragmentAccess::ReadWrite);
+}
+
+void UMassPhysicsBodyInstancesBatchInitializer::Execute(FMassEntityManager& EntityManager, FMassExecutionContext& ExecutionContext)
+{
+	using namespace UE::MassEngine::Private;
+
+	//////////////////////////////////////////////////////////////////////////////
+	// Batch create the mass entities
+	EntityQuery.ForEachEntityChunk(ExecutionContext,
+		[this](FMassExecutionContext& Context)
+		{
+			const TArrayView<FMassPhysicsBodyInstancesInitializationRequestFragment> RequestFragments = Context.GetMutableFragmentView<FMassPhysicsBodyInstancesInitializationRequestFragment>();
+
+			TRACE_CPUPROFILER_EVENT_SCOPE(UMassPhysicsBodyInstancesBatchInitializer::InitAllInstanceBodies);
+			CSV_SCOPED_TIMING_STAT_EXCLUSIVE(UMassPhysicsBodyInstancesBatchInitializer);
+
+			TArray<FMassEntityHandle> EntitiesToDestroy;
+			EntitiesToDestroy.Reserve(Context.GetNumEntities());
+
+			for (FMassExecutionContext::FEntityIterator EntityIt = Context.CreateEntityIterator(); EntityIt; ++EntityIt)
+			{
+				EntitiesToDestroy.Add(Context.GetEntity(*EntityIt));
+
+				FMassPhysicsBodyInstancesInitializationRequestFragment& RequestFragment = RequestFragments[*EntityIt];
+				const int32 NumBodies = RequestFragment.InstanceBodies.Num();
+				if (!ensureMsgf(NumBodies > 0
+					, TEXT("Initialization request fragment was created with an empty list for body instances."
+						" Make sure all code paths properly fill the request.")))
+				{
+					continue;
+				}
+
+				if (!ensureMsgf(RequestFragment.Transforms.Num() == NumBodies
+					, TEXT("Initialization request fragment was created with an incorrect number of transforms. Expecting %d but has %d"
+						" Make sure all code paths properly fill the request.")
+					, NumBodies
+					, RequestFragment.Transforms.Num()))
+				{
+					continue;
+				}
+
+				if (!ensureMsgf(!RequestFragment.StaticMesh.IsExplicitlyNull()
+					, TEXT("Initialization request fragment was created without a valid mesh."
+						" Make sure all code paths properly fill the request.")))
+				{
+					continue;
+				}
+
+				const UStaticMesh* Mesh = RequestFragment.StaticMesh.Get();
+				if (Mesh == nullptr)
+				{
+					UE_VLOG_UELOG(this, LogMassEngine, Log, TEXT("Specified mesh is no longer valid. Ignoring the request"));
+					continue;
+				}
+
+				UBodySetup* BodySetup = Mesh->GetBodySetup();
+				if (BodySetup == nullptr)
+				{
+					UE_VLOG_UELOG(this, LogMassEngine, Log, TEXT("Invalid BodySetup in the specified mesh. Ignoring the request"));
+					continue;
+				}
+
+				if (RequestFragment.BodyInstanceOwner.IsStale())
+				{
+					UE_VLOG_UELOG(this, LogMassEngine, Log, TEXT("Specified instance's owner is no longer valid. Ignoring the request"));
+					continue;
+				}
+
+				// Sanitized array does not contain any nulls
+				TArray<FBodyInstance*> InstanceBodiesSanitized;
+				InstanceBodiesSanitized.Reserve(NumBodies);
+
+				TArray<FTransform> Transforms;
+				Transforms.Reserve(NumBodies);
+
+				EBodyType CommonBodyType = RequestFragment.bDynamic ? EBodyType::NonSimulatedDynamic : EBodyType::Static;
+				for (int32 BodyIndex = 0; BodyIndex < NumBodies; ++BodyIndex)
+				{
+					if (const TSharedPtr<FBodyInstance>& Instance = RequestFragment.InstanceBodies[BodyIndex])
+					{
+						if (Instance->bSimulatePhysics
+							&& CommonBodyType == EBodyType::NonSimulatedDynamic)
+						{
+							CommonBodyType = EBodyType::SimulatedDynamic;
+						}
+						ensure((CommonBodyType == UE::MassEngine::Private::EBodyType::Static && !Instance->bSimulatePhysics)
+							|| (CommonBodyType == UE::MassEngine::Private::EBodyType::SimulatedDynamic && Instance->bSimulatePhysics)
+							|| (CommonBodyType == UE::MassEngine::Private::EBodyType::NonSimulatedDynamic && !Instance->bSimulatePhysics));
+
+						// We need clean up existing static body when getting initialized again with different settings
+						if (Instance->IsValidBodyInstance())
+						{
+							Instance->TermBody();
+						}
+
+						// Instances transforms are in global space
+						const FTransform InstanceTM = FTransform(RequestFragment.Transforms[BodyIndex]);
+
+						// Skip bodies if associated transform is not valid (e.g., scale too small, NaN, mirroring)
+						const bool bTransformIsValid = Instance->ValidateTransform(InstanceTM, TEXT("batch initializer skipped that instance"), BodySetup);
+						if (bTransformIsValid)
+						{
+							if (!Instance->GetOverrideWalkableSlopeOnInstance())
+							{
+								Instance->SetWalkableSlopeOverride(BodySetup->WalkableSlopeOverride, false);
+							}
+							InstanceBodiesSanitized.Add(Instance.Get());
+
+							Transforms.Add(InstanceTM);
+						}
+					}
+				}
+
+				if (InstanceBodiesSanitized.Num() > 0)
+				{
+					if (RequestFragment.bDynamic)
+					{
+						const FInitBodySpawnParams SpawnParams(
+							/*bStaticPhysics*/false,
+							/*bPhysicsTypeDeterminesSimulation*/false);
+
+						constexpr bool bCompileStatic = false;
+						FInitBodiesHelper<bCompileStatic> InitBodyHelper(InstanceBodiesSanitized
+							, Transforms
+							, BodySetup
+							, /*UPrimitiveComponent*/RequestFragment.PrimitiveComponent
+							, GetWorld()->GetPhysicsScene()
+							, SpawnParams
+							, SpawnParams.Aggregate
+							, RequestFragment.BodyInstanceOwner.Get());
+						InitBodyHelper.InitBodies();
+					}
+					else
+					{
+						constexpr bool bCompileStatic = true;
+						FInitBodiesHelper<bCompileStatic> InitBodiesHelper(InstanceBodiesSanitized
+							, Transforms
+							, BodySetup
+							, /*UPrimitiveComponent*/RequestFragment.PrimitiveComponent
+							, GetWorld()->GetPhysicsScene()
+							, FInitBodySpawnParams(/*UPrimitiveComponent*/RequestFragment.PrimitiveComponent)
+							, FPhysicsAggregateHandle()
+							, RequestFragment.BodyInstanceOwner.Get());
+						InitBodiesHelper.InitBodies();
+					}
+
+					UE_VLOG_UELOG(this, LogMassEngine, Log, TEXT("Initialized %d %s body instances using mesh: %s")
+						, NumBodies
+						, CommonBodyType == EBodyType::SimulatedDynamic ? TEXT("simulated dynamic") : CommonBodyType == EBodyType::NonSimulatedDynamic ? TEXT("non-simulated dynamic") : TEXT("static")
+						, *GetPathNameSafe(Mesh));
+
+					const TArray<TSharedPtr<FChaosUserDefinedEntity>>& UserDefinedEntityList = RequestFragment.ChaosUserDefinedEntity;
+					if (RequestFragment.ChaosUserDefinedEntity.Num() > 0)
+					{
+						const bool bSharedUserDefinedEntity = UserDefinedEntityList.Num() == 1 && RequestFragment.InstanceBodies.Num() > 1;
+						if (bSharedUserDefinedEntity)
+						{
+							TArray<Chaos::FPhysicsObject*> PhysicsObjects;
+							PhysicsObjects.Reserve(RequestFragment.InstanceBodies.Num());
+							for (const TSharedPtr<FBodyInstance>& InstancedBody : RequestFragment.InstanceBodies)
+							{
+								if (Chaos::FPhysicsObject* PhysicsObject = InstancedBody && InstancedBody->GetPhysicsActor() ? InstancedBody->GetPhysicsActor()->GetPhysicsObject() : nullptr)
+								{
+									PhysicsObjects.Add(PhysicsObject);
+								}
+							}
+							FPhysicsObjectExternalInterface::LockWrite(PhysicsObjects)->SetUserDefinedEntity(PhysicsObjects, UserDefinedEntityList[0].Get());
+						}
+						else if (ensureMsgf(UserDefinedEntityList.Num() == RequestFragment.InstanceBodies.Num()
+							, TEXT("UserDefinedEntity should be shared or its count should match the number of body instances")))
+						{
+							FLockedWritePhysicsObjectExternalInterface Interface = FPhysicsObjectExternalInterface::LockWrite(GetWorld()->GetPhysicsScene());
+							int32 Index = 0;
+							for (const TSharedPtr<FBodyInstance>& InstancedBody : RequestFragment.InstanceBodies)
+							{
+								if (UserDefinedEntityList[Index])
+								{
+									if (Chaos::FPhysicsObject* PhysicsObject = InstancedBody && InstancedBody->GetPhysicsActor() ? InstancedBody->GetPhysicsActor()->GetPhysicsObject() : nullptr)
+									{
+										Interface->SetUserDefinedEntity({PhysicsObject}, UserDefinedEntityList[Index].Get());
+									}
+								}
+								Index++;
+							}
+						}
+					}
+				}
+			}
+
+
+			Context.Defer().DestroyEntities(EntitiesToDestroy);
+		});
+}
+
+//----------------------------------------------------------------------//
+// UMassPhysicsBodyInstanceReInitializer
+//----------------------------------------------------------------------//
+UMassPhysicsBodyInstanceReInitializer::UMassPhysicsBodyInstanceReInitializer(const FObjectInitializer& ObjectInitializer)
+	: Super(ObjectInitializer)
+	, EntityQuery(*this)
+{
+	ObservedTypes.Append({
+		FMassPhysicsDynamicBodyTag::StaticStruct()
+		, FMassPhysicsSimulatedBodyTag::StaticStruct()
+		, FMassPhysicsCollisionSettingsFragment::StaticStruct() });
+	ObservedOperations = EMassObservedOperationFlags::All;
+	ExecutionFlags = static_cast<int32>(EProcessorExecutionFlags::AllWorldModes);
+}
+
+void UMassPhysicsBodyInstanceReInitializer::ConfigureQueries(const TSharedRef<FMassEntityManager>& EntityManager)
+{
+	EntityQuery.AddTagRequirement<FMassPhysicsBodyInstanceInitializedTag>(EMassFragmentPresence::All);
+}
+
+void UMassPhysicsBodyInstanceReInitializer::Execute(FMassEntityManager& EntityManager, FMassExecutionContext& ExecutionContext)
+{
+	EntityQuery.ForEachEntityChunk(ExecutionContext, [this](const FMassExecutionContext& Context)
+		{
+			Context.GetEntityManagerChecked().Defer().PushCommand<FMassCommandRemoveElements<FMassPhysicsBodyInstanceInitializedTag>>(Context.GetEntities());
+		});
+}
diff --git a/Engine/Source/Runtime/Mass/MassEngine/Private/RigidBody/MassPhysicsProcessors.h b/Engine/Source/Runtime/Mass/MassEngine/Private/RigidBody/MassPhysicsProcessors.h
new file mode 100644
--- /dev/null
+++ b/Engine/Source/Runtime/Mass/MassEngine/Private/RigidBody/MassPhysicsProcessors.h
@@ -0,0 +1,67 @@
+// Copyright Epic Games, Inc. All Rights Reserved.
+
+#pragma once
+
+#include "MassObserverProcessor.h"
+#include "MassProcessor.h"
+#include "MassPhysicsProcessors.generated.h"
+
+struct FMassEntityManager;
+struct FMassExecutionContext;
+
+/**
+ * Processor responsible for batching individual BodyInstance owner into a single initialization request.
+ */
+UCLASS()
+class UMassPhysicsBodyInstancesBatcher : public UMassProcessor
+{
+	GENERATED_BODY()
+
+public:
+	explicit UMassPhysicsBodyInstancesBatcher(const FObjectInitializer& ObjectInitializer = FObjectInitializer::Get());
+
+protected:
+	virtual void ConfigureQueries(const TSharedRef<FMassEntityManager>& EntityManager) override;
+	virtual void Execute(FMassEntityManager& EntityManager, FMassExecutionContext& ExecutionContext) override;
+
+	FMassArchetypeHandle BatchInitializationRequestArchetype;
+	FMassEntityQuery EntityQuery;
+};
+
+/**
+ * Processor responsible for initializing batches of FBodyInstances.
+ */
+UCLASS()
+class UMassPhysicsBodyInstancesBatchInitializer : public UMassProcessor
+{
+	GENERATED_BODY()
+
+public:
+	explicit UMassPhysicsBodyInstancesBatchInitializer(const FObjectInitializer& ObjectInitializer = FObjectInitializer::Get());
+
+protected:
+
+	virtual void ConfigureQueries(const TSharedRef<FMassEntityManager>& EntityManager) override;
+	virtual void Execute(FMassEntityManager& EntityManager, FMassExecutionContext& ExecutionContext) override;
+
+	FMassEntityQuery EntityQuery;
+};
+
+/**
+ * Processor responsible to monitor changes to body instances configurations to trigger a re-initialization of the instance.
+ */
+UCLASS()
+class UMassPhysicsBodyInstanceReInitializer : public UMassObserverProcessor
+{
+	GENERATED_BODY()
+
+public:
+	explicit UMassPhysicsBodyInstanceReInitializer(const FObjectInitializer& ObjectInitializer = FObjectInitializer::Get());
+
+protected:
+
+	virtual void ConfigureQueries(const TSharedRef<FMassEntityManager>& EntityManager) override;
+	virtual void Execute(FMassEntityManager& EntityManager, FMassExecutionContext& ExecutionContext) override;
+
+	FMassEntityQuery EntityQuery;
+};
\ No newline at end of file
diff --git a/Engine/Source/Runtime/Mass/MassEngine/Public/MassEngineTypes.h b/Engine/Source/Runtime/Mass/MassEngine/Public/MassEngineTypes.h
--- a/Engine/Source/Runtime/Mass/MassEngine/Public/MassEngineTypes.h
+++ b/Engine/Source/Runtime/Mass/MassEngine/Public/MassEngineTypes.h
@@ -24,4 +24,7 @@ namespace UE::Mass::ProcessorGroupNames
 {
 	const FName SetupRenderState = FName(TEXT("SetupRenderState"));
 	const FName CreateRenderState = FName(TEXT("CreateRenderState"));
+
+	const FName SetupBatchedInitializations = FName(TEXT("SetupBatchedInitializations"));
+	const FName ProcessBatchedInitializations = FName(TEXT("ProcessBatchedInitializations"));
 }
\ No newline at end of file
diff --git a/Engine/Source/Runtime/Mass/MassEngine/Public/Physics/MassEnginePhysicsFragments.h b/Engine/Source/Runtime/Mass/MassEngine/Public/Physics/MassEnginePhysicsFragments.h
--- a/Engine/Source/Runtime/Mass/MassEngine/Public/Physics/MassEnginePhysicsFragments.h
+++ b/Engine/Source/Runtime/Mass/MassEngine/Public/Physics/MassEnginePhysicsFragments.h
@@ -4,29 +4,146 @@
 
 #include "Engine/EngineTypes.h"
 #include "Mass/EntityElementTypes.h"
-
+#include "Mass/ExternalSubsystemTraits.h"
+#include "UObject/WeakInterfacePtr.h"
 #include "MassEnginePhysicsFragments.generated.h"
 
-class UStaticMesh;
+class FChaosUserDefinedEntity;
+class IPhysicsBodyInstanceOwner;
+struct FBodyInstance;
 
 #define UE_API MASSENGINE_API
 
 /**
- * Fragment holding data related to physics collision
+ * Fragment holding data related to physics collision setup
  */
 USTRUCT()
-struct FMassColliderFragment : public FMassConstSharedFragment
+struct FMassPhysicsCollisionSettingsFragment : public FMassConstSharedFragment
 {
 	GENERATED_BODY()
 
-	FMassColliderFragment() = default;
-	UE_API explicit FMassColliderFragment(TNotNull<const UStaticMesh*> Mesh);
-
-	UPROPERTY(VisibleAnywhere, Category=Debug)
+	UPROPERTY(VisibleAnywhere, Category = Physics)
 	TEnumAsByte<ECollisionEnabled::Type> CollisionType = ECollisionEnabled::NoCollision;
 
-	UPROPERTY(VisibleAnywhere, Category=Debug)
+	UPROPERTY(VisibleAnywhere, Category = Physics)
 	FCollisionResponseContainer CollisionResponse;
 };
 
+/**
+ * Fragment holding data related to a specific rigid body instance
+ */
+USTRUCT()
+struct FMassPhysicsBodyInstanceFragment : public FMassFragment
+{
+	GENERATED_BODY()
+
+	FMassPhysicsBodyInstanceFragment() = default;
+
+	explicit FMassPhysicsBodyInstanceFragment(const TSharedRef<FBodyInstance>& BodyInstance)
+		: BodyInstance(BodyInstance)
+	{
+	}
+
+	/** Weak object pointer to the owning component. Validate before using BodyInstance and ChaosUserDefinedEntity,
+	 * as those shared pointers may outlive the owner and become semantically stale. */
+	UPROPERTY(VisibleAnywhere, Category = Physics)
+	TWeakObjectPtr<UObject> BodyInstanceOwnerObject;
+
+	/** Shared pointer to the BodyInstance created by the owning component. */
+	TSharedPtr<FBodyInstance> BodyInstance = nullptr;
+
+	/**
+	 * Optional payload to associate with the created PhysicsObject
+	 * after BodyInstance initialization (i.e., FWritePhysicsObjectInterface::SetUserDefinedEntity).
+	 */
+	TSharedPtr<FChaosUserDefinedEntity> ChaosUserDefinedEntity = nullptr;
+};
+
+template <>
+struct TMassFragmentTraits<FMassPhysicsBodyInstanceFragment> final
+{
+	enum
+	{
+		AuthorAcceptsItsNotTriviallyCopyable = true
+	};
+};
+
+/**
+ * Tag that indicates that a given body should be initialized as dynamic with the possibility to be simulated in the physics scene
+ * @see FMassPhysicsSimulatedBodyTag
+ */
+USTRUCT()
+struct FMassPhysicsDynamicBodyTag : public FMassTag
+{
+	GENERATED_BODY()
+};
+
+/**
+ * Tag that indicates that a given body should be simulated in the physics scene
+ */
+USTRUCT()
+struct FMassPhysicsSimulatedBodyTag : public FMassTag
+{
+	GENERATED_BODY()
+};
+
+/**
+ * Tag added after a body instance has been batched for initialization.
+ * Entities with this tag are excluded from further batching. Removed by
+ * UMassPhysicsBodyInstanceReInitializer to re-trigger the initialization cycle.
+ */
+USTRUCT()
+struct FMassPhysicsBodyInstanceInitializedTag : public FMassTag
+{
+	GENERATED_BODY()
+};
+
+USTRUCT()
+struct FMassPhysicsBodyInstancesInitializationRequestFragment : public FMassFragment
+{
+	GENERATED_BODY()
+
+	/** Mandatory weak pointer to the mesh used by all instances */
+	UPROPERTY(Transient, VisibleAnywhere, Category = Physics)
+	TWeakObjectPtr<const UStaticMesh> StaticMesh = nullptr;
+
+	/** Mandatory list of transforms that should match the number of body instances */
+	UPROPERTY(Transient, VisibleAnywhere, Category = Physics)
+	TArray<FTransform> Transforms;
+
+	/** Mandatory list of body instances to initialize */
+	TArray<TSharedPtr<FBodyInstance>> InstanceBodies;
+
+	/**
+	 * Optional interface pointer that will be provided to the BodyInstance initialization
+	 */
+	UPROPERTY(Transient, VisibleAnywhere, Category = Physics)
+	TObjectPtr<UPrimitiveComponent> PrimitiveComponent;
+
+	/**
+	 * Optional interface pointer that will be provided to the BodyInstance initialization.
+	 * Stored as a TWeakInterfacePtr; check IsStale() before use.
+	 */
+	TWeakInterfacePtr<IPhysicsBodyInstanceOwner> BodyInstanceOwner;
+
+	/**
+	 * Optional list of user data payloads to associate to the created PhysicsObjects after the body instances get initialized.
+	 * @note It is possible to provide a single payload for owners relying on body indices (i.e., InstanceBodyIndex).
+	 */
+	TArray<TSharedPtr<FChaosUserDefinedEntity>> ChaosUserDefinedEntity;
+
+	/** Whether the instances are simulated/dynamic physics objects */
+	UPROPERTY(Transient, VisibleAnywhere, Category = Physics)
+	bool bDynamic = false;
+};
+
+template <>
+struct TMassFragmentTraits<FMassPhysicsBodyInstancesInitializationRequestFragment> final
+{
+	enum
+	{
+		AuthorAcceptsItsNotTriviallyCopyable = true
+	};
+};
+
 #undef UE_API
\ No newline at end of file
