[Backout] - CL52540129
#rnx
Original CL Desc
-----------------------------------------------------------------
[MassEngine][RigidPhysicsCore] Batched BodyInstance initialization via Mass processors for SceneGraph entities

Added MassPhysics processors to batch-initialize FBodyInstances through the Mass pipeline, replacing per-entity immediate initialization for SceneGraph entities.

New Mass processors:
- UMassPhysicsBodyInstancesBatcher: groups entities by mesh/collision/body-type into batch requests
- UMassPhysicsBodyInstancesBatchInitializer: calls FInitBodiesHelper on batched requests
- UMassPhysicsBodyInstanceReInitializer: observer that re-queues initialization on config change

New fragments/tags:
- FMassPhysicsCollisionSettingsFragment (replaces FMassColliderFragment)
- FMassPhysicsBodyInstanceFragment, FMassPhysicsBodyInstancesInitializationRequestFragment
- FMassPhysicsDynamicBodyTag, FMassPhysicsSimulatedBodyTag, FMassPhysicsBodyInstanceInitializedTag

FVerseRigidBody changes:
- BodyInstance and ChaosUserData converted from value to TSharedPtr for shared ownership with Mass fragments
- Added "strictly static" concept based on mesh part availability (no mesh parts = strictly static) and use the optimized FInitBodiesHelper</*bCompileStatic*/> initialization path
- Collision settings management centralized in VerseRigidBody (UpdateCollisionSettingsFragment)

Gated with CVar: "SceneGraph.EnableBodyInstancesInitializationInMass" (default: disabled)

Body type classification:
- "strictly static": no mesh parts available
- "non-simulated dynamic": has mesh parts, dynamic flag off
- "simulated dynamic": has mesh parts, dynamic flag on

Runtime transitions supported:
- Toggling the dynamic flag transitions between non-simulated and simulated dynamic
- Gaining mesh parts at runtime transitions from strictly static to dynamic
- Losing mesh parts transitions from dynamic back to strictly static

#rnx

#Synced-CL 52540129

diff --git a/Engine/Source/Runtime/Mass/MassEngine/MassEngine.Build.cs b/Engine/Source/Runtime/Mass/MassEngine/MassEngine.Build.cs
--- a/Engine/Source/Runtime/Mass/MassEngine/MassEngine.Build.cs
+++ b/Engine/Source/Runtime/Mass/MassEngine/MassEngine.Build.cs
@@ -14,7 +14,6 @@ namespace UnrealBuildTool.Rules
 
 			PublicDependencyModuleNames.AddRange(
 				new string[] {
-					"Chaos",
 					"Core",
 					"CoreUObject",
 					"Engine",
diff --git a/Engine/Source/Runtime/Mass/MassEngine/Private/AI/MassEngineNavigationProcessors.cpp b/Engine/Source/Runtime/Mass/MassEngine/Private/AI/MassEngineNavigationProcessors.cpp
--- a/Engine/Source/Runtime/Mass/MassEngine/Private/AI/MassEngineNavigationProcessors.cpp
+++ b/Engine/Source/Runtime/Mass/MassEngine/Private/AI/MassEngineNavigationProcessors.cpp
@@ -22,7 +22,7 @@
 namespace UE::NavigationRelevant
 {
 
-bool IsNavigationRelevant(const FMassPhysicsCollisionSettingsFragment& CollisionSettingsFragment,
+bool IsNavigationRelevant(const FMassColliderFragment& ColliderFragment,
 	const FMassStaticMeshFragment& MeshFragment,
 	const FTransformFragment& TransformFragment)
 {
@@ -33,12 +33,12 @@ bool IsNavigationRelevant(const FMassPhysicsCollisionSettingsFragment& Collision
 
 	// @see UPrimitiveComponent::IsNavigationRelevant()
 	{
-		if (!CollisionEnabledHasQuery(CollisionSettingsFragment.CollisionType))
+		if (!CollisionEnabledHasQuery(ColliderFragment.CollisionType))
 		{
 			return false;
 		}
 
-		const FCollisionResponseContainer& ResponseToChannels = CollisionSettingsFragment.CollisionResponse;
+		const FCollisionResponseContainer& ResponseToChannels = ColliderFragment.CollisionResponse;
 		if (!(ResponseToChannels.GetResponse(ECC_Pawn) == ECR_Block
 			|| ResponseToChannels.GetResponse(ECC_Vehicle) == ECR_Block))
 		{
@@ -157,7 +157,7 @@ UMassNavigationElementProcessor::UMassNavigationElementProcessor(const FObjectIn
 
 void UMassNavigationElementProcessor::ConfigureQueries(const TSharedRef<FMassEntityManager>& EntityManager)
 {
-	NewRelevantEntityQuery.AddConstSharedRequirement<FMassPhysicsCollisionSettingsFragment>();
+	NewRelevantEntityQuery.AddConstSharedRequirement<FMassColliderFragment>();
 	NewRelevantEntityQuery.AddConstSharedRequirement<FMassStaticMeshFragment>();
 	NewRelevantEntityQuery.AddRequirement<FTransformFragment>(EMassFragmentAccess::ReadOnly);
 	NewRelevantEntityQuery.AddRequirement<FMassNavigationRelevantFragment>(EMassFragmentAccess::ReadOnly, EMassFragmentPresence::None);
@@ -180,21 +180,20 @@ void UMassNavigationElementProcessor::Execute(FMassEntityManager& EntityManager,
 		{
 			const TConstArrayView<FTransformFragment> TransformFragments = Context.GetFragmentView<FTransformFragment>();
 			const FMassStaticMeshFragment& MeshFragment = Context.GetConstSharedFragment<FMassStaticMeshFragment>();
-			const FMassPhysicsCollisionSettingsFragment& CollisionSettingsFragment = Context.GetConstSharedFragment<FMassPhysicsCollisionSettingsFragment>();
+			const FMassColliderFragment& ColliderFragment = Context.GetConstSharedFragment<FMassColliderFragment>();
 
 			for (FMassExecutionContext::FEntityIterator EntityIt = Context.CreateEntityIterator(); EntityIt; ++EntityIt)
 			{
 				const FMassEntityHandle EntityHandle = Context.GetEntity(EntityIt);
 				FNavigationElementHandle NavigationElementHandle = FNavigationElementHandle::Invalid;
 
-				if (!IsNavigationRelevant(CollisionSettingsFragment, MeshFragment, TransformFragments[EntityIt]))
+				if (IsNavigationRelevant(ColliderFragment, MeshFragment, TransformFragments[EntityIt]))
 				{
-					continue;
+					FNavigationElement Element(*EntityManagerOwner, EntityHandle.AsNumber());
+					ConfigureElement(EntityManagerOwner, Element, MeshFragment, TransformFragments[EntityIt], ParametersSharedFragment.Get<const FMassNavigationRelevantParameters>());
+					NavigationElementHandle = FNavigationSystem::AddNavigationElement(EntityManagerOwner->GetWorld(), MoveTemp(Element));
 				}
 
-				FNavigationElement Element(*EntityManagerOwner, EntityHandle.AsNumber());
-				ConfigureElement(EntityManagerOwner, Element, MeshFragment, TransformFragments[EntityIt], ParametersSharedFragment.Get<const FMassNavigationRelevantParameters>());
-				NavigationElementHandle = FNavigationSystem::AddNavigationElement(EntityManagerOwner->GetWorld(), MoveTemp(Element));
 				Context.Defer().PushCommand<FMassCommandAddFragmentInstances>(EntityHandle, FMassNavigationRelevantFragment(NavigationElementHandle));
 				Context.Defer().PushCommand<FMassDeferredAddCommand>([EntityHandle, ParametersSharedFragment](FMassEntityManager& InEntityManager)
 					{
@@ -237,7 +236,7 @@ void UMassDirtyNavigationRelevantUpdateProcessor::InitializeInternal(UObject& Ow
 
 void UMassDirtyNavigationRelevantUpdateProcessor::ConfigureQueries(const TSharedRef<FMassEntityManager>& EntityManager)
 {
-	EntityQuery.AddConstSharedRequirement<FMassPhysicsCollisionSettingsFragment>();
+	EntityQuery.AddConstSharedRequirement<FMassColliderFragment>();
 	EntityQuery.AddConstSharedRequirement<FMassStaticMeshFragment>();
 	EntityQuery.AddConstSharedRequirement<FMassNavigationRelevantParameters>();
 	EntityQuery.AddRequirement<FTransformFragment>(EMassFragmentAccess::ReadOnly);
@@ -260,7 +259,7 @@ void UMassDirtyNavigationRelevantUpdateProcessor::SignalEntities(
 			const TConstArrayView<FTransformFragment> TransformFragments = Context.GetFragmentView<FTransformFragment>();
 			const TArrayView<FMassNavigationRelevantFragment> NavigationRelevantFragments = Context.GetMutableFragmentView<FMassNavigationRelevantFragment>();
 
-			const FMassPhysicsCollisionSettingsFragment& CollisionSettingsFragment = Context.GetConstSharedFragment<FMassPhysicsCollisionSettingsFragment>();
+			const FMassColliderFragment& ColliderFragment = Context.GetConstSharedFragment<FMassColliderFragment>();
 			const FMassStaticMeshFragment& MeshFragment = Context.GetConstSharedFragment<FMassStaticMeshFragment>();
 			const FMassNavigationRelevantParameters& NavigationRelevantParameters = Context.GetConstSharedFragment<FMassNavigationRelevantParameters>();
 
@@ -268,7 +267,7 @@ void UMassDirtyNavigationRelevantUpdateProcessor::SignalEntities(
 			{
 				const bool bWasRelevant = NavigationRelevantFragments[EntityIt].Handle.IsValid();
 
-				if (IsNavigationRelevant(CollisionSettingsFragment, MeshFragment, TransformFragments[EntityIt]))
+				if (IsNavigationRelevant(ColliderFragment, MeshFragment, TransformFragments[EntityIt]))
 				{
 					FNavigationElement Element(*EntityManagerOwner, Context.GetEntity(EntityIt).AsNumber());
 					ConfigureElement(EntityManagerOwner, Element, MeshFragment, TransformFragments[EntityIt], NavigationRelevantParameters);
@@ -297,19 +296,13 @@ UMassNavigationRelevantMassDeinitializer::UMassNavigationRelevantMassDeinitializ
 	: EntityQuery(*this)
 {
 	ExecutionFlags = static_cast<int32>(EProcessorExecutionFlags::AllWorldModes);
-	ObservedTypes.Append({
-		FMassStaticMeshFragment::StaticStruct(),
-		FTransformFragment::StaticStruct(),
-		FMassPhysicsCollisionSettingsFragment::StaticStruct()});
+	ObservedTypes.Add(FMassNavigationRelevantFragment::StaticStruct());
 	ObservedOperations = EMassObservedOperationFlags::Remove;
 	bRequiresGameThreadExecution = true;
 }
 
 void UMassNavigationRelevantMassDeinitializer::ConfigureQueries(const TSharedRef<FMassEntityManager>& EntityManager)
 {
-	EntityQuery.AddConstSharedRequirement<FMassStaticMeshFragment>(EMassFragmentPresence::All);
-	EntityQuery.AddConstSharedRequirement<FMassPhysicsCollisionSettingsFragment>(EMassFragmentPresence::All);
-	EntityQuery.AddRequirement<FTransformFragment>(EMassFragmentAccess::ReadOnly);
 	EntityQuery.AddRequirement<FMassNavigationRelevantFragment>(EMassFragmentAccess::ReadOnly);
 }
 
@@ -326,7 +319,5 @@ void UMassNavigationRelevantMassDeinitializer::Execute(FMassEntityManager& Entit
 					FNavigationSystem::RemoveNavigationElement(World, NavigationRelevantFragments[EntityIt].Handle);
 				}
 			}
-
-			Context.Defer().PushCommand<FMassCommandRemoveElements<FMassNavigationRelevantFragment, FMassNavigationRelevantParameters>>(Context.GetEntities());
 		});
 }
diff --git a/Engine/Source/Runtime/Mass/MassEngine/Private/Mesh/MassStaticMeshRenderStateHelper.cpp b/Engine/Source/Runtime/Mass/MassEngine/Private/Mesh/MassStaticMeshRenderStateHelper.cpp
--- a/Engine/Source/Runtime/Mass/MassEngine/Private/Mesh/MassStaticMeshRenderStateHelper.cpp
+++ b/Engine/Source/Runtime/Mass/MassEngine/Private/Mesh/MassStaticMeshRenderStateHelper.cpp
@@ -54,9 +54,9 @@ UMaterialInterface* FMassBaseStaticMeshRenderStateHelper::GetMaterial(int32 Mate
 
 const FCollisionResponseContainer& FMassBaseStaticMeshRenderStateHelper::GetCollisionResponseToChannels() const
 {
-	if (FMassPhysicsCollisionSettingsFragment* CollisionSettingsFragment = GetEntityManager().GetConstSharedFragmentDataPtr<FMassPhysicsCollisionSettingsFragment>(EntityHandle))
+	if (FMassColliderFragment* ColliderFragment = GetEntityManager().GetConstSharedFragmentDataPtr<FMassColliderFragment>(EntityHandle))
 	{
-		return CollisionSettingsFragment->CollisionResponse;
+		return ColliderFragment->CollisionResponse;
 	}
 
 	return FCollisionResponseContainer::GetDefaultResponseContainer();
diff --git a/Engine/Source/Runtime/Mass/MassEngine/Private/RigidBody/MassPhysicsProcessors.cpp b/Engine/Source/Runtime/Mass/MassEngine/Private/RigidBody/MassPhysicsProcessors.cpp
deleted file mode 100644
diff --git a/Engine/Source/Runtime/Mass/MassEngine/Private/RigidBody/MassPhysicsProcessors.h b/Engine/Source/Runtime/Mass/MassEngine/Private/RigidBody/MassPhysicsProcessors.h
deleted file mode 100644
diff --git a/Engine/Source/Runtime/Mass/MassEngine/Public/MassEngineTypes.h b/Engine/Source/Runtime/Mass/MassEngine/Public/MassEngineTypes.h
--- a/Engine/Source/Runtime/Mass/MassEngine/Public/MassEngineTypes.h
+++ b/Engine/Source/Runtime/Mass/MassEngine/Public/MassEngineTypes.h
@@ -24,7 +24,4 @@ namespace UE::Mass::ProcessorGroupNames
 {
 	const FName SetupRenderState = FName(TEXT("SetupRenderState"));
 	const FName CreateRenderState = FName(TEXT("CreateRenderState"));
-
-	const FName SetupBatchedInitializations = FName(TEXT("SetupBatchedInitializations"));
-	const FName ProcessBatchedInitializations = FName(TEXT("ProcessBatchedInitializations"));
 }
\ No newline at end of file
diff --git a/Engine/Source/Runtime/Mass/MassEngine/Public/Physics/MassEnginePhysicsFragments.h b/Engine/Source/Runtime/Mass/MassEngine/Public/Physics/MassEnginePhysicsFragments.h
--- a/Engine/Source/Runtime/Mass/MassEngine/Public/Physics/MassEnginePhysicsFragments.h
+++ b/Engine/Source/Runtime/Mass/MassEngine/Public/Physics/MassEnginePhysicsFragments.h
@@ -4,146 +4,29 @@
 
 #include "Engine/EngineTypes.h"
 #include "Mass/EntityElementTypes.h"
-#include "Mass/ExternalSubsystemTraits.h"
-#include "UObject/WeakInterfacePtr.h"
+
 #include "MassEnginePhysicsFragments.generated.h"
 
-class FChaosUserDefinedEntity;
-class IPhysicsBodyInstanceOwner;
-struct FBodyInstance;
+class UStaticMesh;
 
 #define UE_API MASSENGINE_API
 
 /**
- * Fragment holding data related to physics collision setup
+ * Fragment holding data related to physics collision
  */
 USTRUCT()
-struct FMassPhysicsCollisionSettingsFragment : public FMassConstSharedFragment
+struct FMassColliderFragment : public FMassConstSharedFragment
 {
 	GENERATED_BODY()
 
-	UPROPERTY(VisibleAnywhere, Category = Physics)
+	FMassColliderFragment() = default;
+	UE_API explicit FMassColliderFragment(TNotNull<const UStaticMesh*> Mesh);
+
+	UPROPERTY(VisibleAnywhere, Category=Debug)
 	TEnumAsByte<ECollisionEnabled::Type> CollisionType = ECollisionEnabled::NoCollision;
 
-	UPROPERTY(VisibleAnywhere, Category = Physics)
+	UPROPERTY(VisibleAnywhere, Category=Debug)
 	FCollisionResponseContainer CollisionResponse;
 };
 
-/**
- * Fragment holding data related to a specific rigid body instance
- */
-USTRUCT()
-struct FMassPhysicsBodyInstanceFragment : public FMassFragment
-{
-	GENERATED_BODY()
-
-	FMassPhysicsBodyInstanceFragment() = default;
-
-	explicit FMassPhysicsBodyInstanceFragment(const TSharedRef<FBodyInstance>& BodyInstance)
-		: BodyInstance(BodyInstance)
-	{
-	}
-
-	/** Weak object pointer to the owning component. Validate before using BodyInstance and ChaosUserDefinedEntity,
-	 * as those shared pointers may outlive the owner and become semantically stale. */
-	UPROPERTY(VisibleAnywhere, Category = Physics)
-	TWeakObjectPtr<UObject> BodyInstanceOwnerObject;
-
-	/** Shared pointer to the BodyInstance created by the owning component. */
-	TSharedPtr<FBodyInstance> BodyInstance = nullptr;
-
-	/**
-	 * Optional payload to associate with the created PhysicsObject
-	 * after BodyInstance initialization (i.e., FWritePhysicsObjectInterface::SetUserDefinedEntity).
-	 */
-	TSharedPtr<FChaosUserDefinedEntity> ChaosUserDefinedEntity = nullptr;
-};
-
-template <>
-struct TMassFragmentTraits<FMassPhysicsBodyInstanceFragment> final
-{
-	enum
-	{
-		AuthorAcceptsItsNotTriviallyCopyable = true
-	};
-};
-
-/**
- * Tag that indicates that a given body should be initialized as dynamic with the possibility to be simulated in the physics scene
- * @see FMassPhysicsSimulatedBodyTag
- */
-USTRUCT()
-struct FMassPhysicsDynamicBodyTag : public FMassTag
-{
-	GENERATED_BODY()
-};
-
-/**
- * Tag that indicates that a given body should be simulated in the physics scene
- */
-USTRUCT()
-struct FMassPhysicsSimulatedBodyTag : public FMassTag
-{
-	GENERATED_BODY()
-};
-
-/**
- * Tag added after a body instance has been batched for initialization.
- * Entities with this tag are excluded from further batching. Removed by
- * UMassPhysicsBodyInstanceReInitializer to re-trigger the initialization cycle.
- */
-USTRUCT()
-struct FMassPhysicsBodyInstanceInitializedTag : public FMassTag
-{
-	GENERATED_BODY()
-};
-
-USTRUCT()
-struct FMassPhysicsBodyInstancesInitializationRequestFragment : public FMassFragment
-{
-	GENERATED_BODY()
-
-	/** Mandatory weak pointer to the mesh used by all instances */
-	UPROPERTY(Transient, VisibleAnywhere, Category = Physics)
-	TWeakObjectPtr<const UStaticMesh> StaticMesh = nullptr;
-
-	/** Mandatory list of transforms that should match the number of body instances */
-	UPROPERTY(Transient, VisibleAnywhere, Category = Physics)
-	TArray<FTransform> Transforms;
-
-	/** Mandatory list of body instances to initialize */
-	TArray<TSharedPtr<FBodyInstance>> InstanceBodies;
-
-	/**
-	 * Optional interface pointer that will be provided to the BodyInstance initialization
-	 */
-	UPROPERTY(Transient, VisibleAnywhere, Category = Physics)
-	TObjectPtr<UPrimitiveComponent> PrimitiveComponent;
-
-	/**
-	 * Optional interface pointer that will be provided to the BodyInstance initialization.
-	 * Stored as a TWeakInterfacePtr; check IsStale() before use.
-	 */
-	TWeakInterfacePtr<IPhysicsBodyInstanceOwner> BodyInstanceOwner;
-
-	/**
-	 * Optional list of user data payloads to associate to the created PhysicsObjects after the body instances get initialized.
-	 * @note It is possible to provide a single payload for owners relying on body indices (i.e., InstanceBodyIndex).
-	 */
-	TArray<TSharedPtr<FChaosUserDefinedEntity>> ChaosUserDefinedEntity;
-
-	/** Whether the instances are simulated/dynamic physics objects */
-	UPROPERTY(Transient, VisibleAnywhere, Category = Physics)
-	bool bDynamic = false;
-};
-
-template <>
-struct TMassFragmentTraits<FMassPhysicsBodyInstancesInitializationRequestFragment> final
-{
-	enum
-	{
-		AuthorAcceptsItsNotTriviallyCopyable = true
-	};
-};
-
 #undef UE_API
\ No newline at end of file
