Fix CIS Issue 1299046 (PVS) -  fixed hash building in MassPhysicsProcessors

#rnx
#Synced-CL 52558571

diff --git a/Engine/Source/Runtime/Mass/MassEngine/Private/RigidBody/MassPhysicsProcessors.cpp b/Engine/Source/Runtime/Mass/MassEngine/Private/RigidBody/MassPhysicsProcessors.cpp
--- a/Engine/Source/Runtime/Mass/MassEngine/Private/RigidBody/MassPhysicsProcessors.cpp
+++ b/Engine/Source/Runtime/Mass/MassEngine/Private/RigidBody/MassPhysicsProcessors.cpp
@@ -100,12 +100,7 @@ void UMassPhysicsBodyInstancesBatcher::Execute(FMassEntityManager& EntityManager
 
 			// Build the source shared fragment hash
 			FHashBuilder HashBuilder;
-			auto AddFragmentToHash = [&HashBuilder](const void* FragmentPtr)
-				{
-					HashBuilder.AppendRaw(&FragmentPtr, sizeof(&FragmentPtr));
-				};
-			AddFragmentToHash(&StaticMeshFragment);
-			HashBuilder.Append(BodyType);
+			HashBuilder << &StaticMeshFragment << BodyType;
 
 			// Find if the batch info already exists, if not add it
 			FBatchInfo* BatchInfo;
