[MassEngine] Fix navigation relevance re-evaluation every frame for non-relevant entities

Non-navigation-relevant entities were being re-evaluated by UMassNavigationElementProcessor every frame because they were never marked as processed. The query filtered on FMassNavigationRelevantFragment absence, but only relevant entities received the fragment.

Added FMassNavigationEvaluatedTag to mark all entities after initial evaluation regardless of relevance outcome. The query now filters on the tag absence instead, so non-relevant entities are no longer re-evaluated. The tag is removed by UMassNavigationRelevantMassDeinitializer alongside other navigation fragments, allowing re-evaluation if an entity's collision or mesh configuration changes.

#Synced-CL 52558571

diff --git a/Engine/Source/Runtime/Mass/MassEngine/Private/AI/MassEngineNavigationProcessors.cpp b/Engine/Source/Runtime/Mass/MassEngine/Private/AI/MassEngineNavigationProcessors.cpp
--- a/Engine/Source/Runtime/Mass/MassEngine/Private/AI/MassEngineNavigationProcessors.cpp
+++ b/Engine/Source/Runtime/Mass/MassEngine/Private/AI/MassEngineNavigationProcessors.cpp
@@ -160,7 +160,7 @@ void UMassNavigationElementProcessor::ConfigureQueries(const TSharedRef<FMassEnt
 	NewRelevantEntityQuery.AddConstSharedRequirement<FMassPhysicsCollisionSettingsFragment>();
 	NewRelevantEntityQuery.AddConstSharedRequirement<FMassStaticMeshFragment>();
 	NewRelevantEntityQuery.AddRequirement<FTransformFragment>(EMassFragmentAccess::ReadOnly);
-	NewRelevantEntityQuery.AddRequirement<FMassNavigationRelevantFragment>(EMassFragmentAccess::ReadOnly, EMassFragmentPresence::None);
+	NewRelevantEntityQuery.AddTagRequirement<FMassNavigationEvaluatedTag>(EMassFragmentPresence::None);
 }
 
 void UMassNavigationElementProcessor::Execute(FMassEntityManager& EntityManager, FMassExecutionContext& ExecutionContext)
@@ -185,7 +185,9 @@ void UMassNavigationElementProcessor::Execute(FMassEntityManager& EntityManager,
 			for (FMassExecutionContext::FEntityIterator EntityIt = Context.CreateEntityIterator(); EntityIt; ++EntityIt)
 			{
 				const FMassEntityHandle EntityHandle = Context.GetEntity(EntityIt);
-				FNavigationElementHandle NavigationElementHandle = FNavigationElementHandle::Invalid;
+
+				// Mark as evaluated so this entity is excluded from future runs
+				Context.Defer().AddTag<FMassNavigationEvaluatedTag>(EntityHandle);
 
 				if (!IsNavigationRelevant(CollisionSettingsFragment, MeshFragment, TransformFragments[EntityIt]))
 				{
@@ -194,7 +196,7 @@ void UMassNavigationElementProcessor::Execute(FMassEntityManager& EntityManager,
 
 				FNavigationElement Element(*EntityManagerOwner, EntityHandle.AsNumber());
 				ConfigureElement(EntityManagerOwner, Element, MeshFragment, TransformFragments[EntityIt], ParametersSharedFragment.Get<const FMassNavigationRelevantParameters>());
-				NavigationElementHandle = FNavigationSystem::AddNavigationElement(EntityManagerOwner->GetWorld(), MoveTemp(Element));
+				const FNavigationElementHandle NavigationElementHandle = FNavigationSystem::AddNavigationElement(EntityManagerOwner->GetWorld(), MoveTemp(Element));
 				Context.Defer().PushCommand<FMassCommandAddFragmentInstances>(EntityHandle, FMassNavigationRelevantFragment(NavigationElementHandle));
 				Context.Defer().PushCommand<FMassDeferredAddCommand>([EntityHandle, ParametersSharedFragment](FMassEntityManager& InEntityManager)
 					{
@@ -327,6 +329,6 @@ void UMassNavigationRelevantMassDeinitializer::Execute(FMassEntityManager& Entit
 				}
 			}
 
-			Context.Defer().PushCommand<FMassCommandRemoveElements<FMassNavigationRelevantFragment, FMassNavigationRelevantParameters>>(Context.GetEntities());
+			Context.Defer().PushCommand<FMassCommandRemoveElements<FMassNavigationRelevantFragment, FMassNavigationRelevantParameters, FMassNavigationEvaluatedTag>>(Context.GetEntities());
 		});
 }
diff --git a/Engine/Source/Runtime/Mass/MassEngine/Public/AI/MassEngineNavigationFragments.h b/Engine/Source/Runtime/Mass/MassEngine/Public/AI/MassEngineNavigationFragments.h
--- a/Engine/Source/Runtime/Mass/MassEngine/Public/AI/MassEngineNavigationFragments.h
+++ b/Engine/Source/Runtime/Mass/MassEngine/Public/AI/MassEngineNavigationFragments.h
@@ -21,6 +21,17 @@ struct FMassNavigationRelevantParameters : public FMassConstSharedFragment
 	bool bFillCollisionUnderneathForNavData = false;
 };
 
+/**
+ * Tag indicating that the entity has been evaluated for navigation relevance.
+ * Added to all entities after initial evaluation regardless of relevance outcome,
+ * so that non-relevant entities are not re-evaluated every frame.
+ */
+USTRUCT()
+struct FMassNavigationEvaluatedTag : public FMassTag
+{
+	GENERATED_BODY()
+};
+
 /**
  * Fragment holding the registration handle to the navigation element created from a Mass entity.
  * The fragment is added to indicate that a Mass entity is relevant to the AI navigation system.
