[MassEngine]
* Fixed non unity compile error

#Synced-CL 52605890

diff --git a/Engine/Source/Runtime/Mass/MassEngine/Private/RigidBody/MassPhysicsProcessors.h b/Engine/Source/Runtime/Mass/MassEngine/Private/RigidBody/MassPhysicsProcessors.h
--- a/Engine/Source/Runtime/Mass/MassEngine/Private/RigidBody/MassPhysicsProcessors.h
+++ b/Engine/Source/Runtime/Mass/MassEngine/Private/RigidBody/MassPhysicsProcessors.h
@@ -2,6 +2,7 @@
 
 #pragma once
 
+#include "MassEntityQuery.h"
 #include "MassObserverProcessor.h"
 #include "MassProcessor.h"
 #include "MassPhysicsProcessors.generated.h"
