[MassEngine] Fix navigation relevance re-evaluation every frame for non-relevant entities

Non-navigation-relevant entities were being re-evaluated by UMassNavigationElementProcessor every frame because they were never marked as processed. The query filtered on FMassNavigationRelevantFragment absence, but only relevant entities received the fragment.

Added FMassNavigationEvaluatedTag to mark all entities after initial evaluation regardless of relevance outcome. The query now filters on the tag absence instead, so non-relevant entities are no longer re-evaluated. UMassNavigationRelevantMassDeinitializer removes the tag from both relevant and non-relevant entities when their observed fragments are removed, allowing re-evaluation if an entity's collision or mesh configuration changes.

#was-backout 52611633
#rnx
#Synced-CL 52704012

diff --git a/Engine/Source/Runtime/Mass/MassEngine/Private/AI/MassEngineNavigationProcessors.cpp b/Engine/Source/Runtime/Mass/MassEngine/Private/AI/MassEngineNavigationProcessors.cpp
--- a/Engine/Source/Runtime/Mass/MassEngine/Private/AI/MassEngineNavigationProcessors.cpp
+++ b/Engine/Source/Runtime/Mass/MassEngine/Private/AI/MassEngineNavigationProcessors.cpp
@@ -160,7 +160,7 @@ void UMassNavigationElementProcessor::ConfigureQueries(const TSharedRef<FMassEnt
 	NewRelevantEntityQuery.AddConstSharedRequirement<FMassPhysicsCollisionSettingsFragment>();
 	NewRelevantEntityQuery.AddConstSharedRequirement<FMassStaticMeshFragment>();
 	NewRelevantEntityQuery.AddRequirement<FTransformFragment>(EMassFragmentAccess::ReadOnly);
-	NewRelevantEntityQuery.AddRequirement<FMassNavigationRelevantFragment>(EMassFragmentAccess::ReadOnly, EMassFragmentPresence::None);
+	NewRelevantEntityQuery.AddTagRequirement<FMassNavigationEvaluatedTag>(EMassFragmentPresence::None);
 }
 
 void UMassNavigationElementProcessor::Execute(FMassEntityManager& EntityManager, FMassExecutionContext& ExecutionContext)
@@ -185,7 +185,9 @@ void UMassNavigationElementProcessor::Execute(FMassEntityManager& EntityManager,
 			for (FMassExecutionContext::FEntityIterator EntityIt = Context.CreateEntityIterator(); EntityIt; ++EntityIt)
 			{
 				const FMassEntityHandle EntityHandle = Context.GetEntity(EntityIt);
-				FNavigationElementHandle NavigationElementHandle = FNavigationElementHandle::Invalid;
+
+				// Mark as evaluated so this entity is excluded from future runs
+				Context.Defer().AddTag<FMassNavigationEvaluatedTag>(EntityHandle);
 
 				if (!IsNavigationRelevant(CollisionSettingsFragment, MeshFragment, TransformFragments[EntityIt]))
 				{
@@ -194,7 +196,7 @@ void UMassNavigationElementProcessor::Execute(FMassEntityManager& EntityManager,
 
 				FNavigationElement Element(*EntityManagerOwner, EntityHandle.AsNumber());
 				ConfigureElement(EntityManagerOwner, Element, MeshFragment, TransformFragments[EntityIt], ParametersSharedFragment.Get<const FMassNavigationRelevantParameters>());
-				NavigationElementHandle = FNavigationSystem::AddNavigationElement(EntityManagerOwner->GetWorld(), MoveTemp(Element));
+				const FNavigationElementHandle NavigationElementHandle = FNavigationSystem::AddNavigationElement(EntityManagerOwner->GetWorld(), MoveTemp(Element));
 				Context.Defer().PushCommand<FMassCommandAddFragmentInstances>(EntityHandle, FMassNavigationRelevantFragment(NavigationElementHandle));
 				Context.Defer().PushCommand<FMassDeferredAddCommand>([EntityHandle, ParametersSharedFragment](FMassEntityManager& InEntityManager)
 					{
@@ -295,6 +297,7 @@ void UMassDirtyNavigationRelevantUpdateProcessor::SignalEntities(
 //----------------------------------------------------------------------//
 UMassNavigationRelevantMassDeinitializer::UMassNavigationRelevantMassDeinitializer()
 	: EntityQuery(*this)
+	, NonRelevantEntityQuery(*this)
 {
 	ExecutionFlags = static_cast<int32>(EProcessorExecutionFlags::AllWorldModes);
 	ObservedTypes.Append({
@@ -311,6 +314,9 @@ void UMassNavigationRelevantMassDeinitializer::ConfigureQueries(const TSharedRef
 	EntityQuery.AddConstSharedRequirement<FMassPhysicsCollisionSettingsFragment>(EMassFragmentPresence::All);
 	EntityQuery.AddRequirement<FTransformFragment>(EMassFragmentAccess::ReadOnly);
 	EntityQuery.AddRequirement<FMassNavigationRelevantFragment>(EMassFragmentAccess::ReadOnly);
+
+	NonRelevantEntityQuery.AddTagRequirement<FMassNavigationEvaluatedTag>(EMassFragmentPresence::All);
+	NonRelevantEntityQuery.AddRequirement<FMassNavigationRelevantFragment>(EMassFragmentAccess::ReadOnly, EMassFragmentPresence::None);
 }
 
 void UMassNavigationRelevantMassDeinitializer::Execute(FMassEntityManager& EntityManager, FMassExecutionContext& ExecutionContext)
@@ -327,6 +333,13 @@ void UMassNavigationRelevantMassDeinitializer::Execute(FMassEntityManager& Entit
 				}
 			}
 
-			Context.Defer().PushCommand<FMassCommandRemoveElements<FMassNavigationRelevantFragment, FMassNavigationRelevantParameters>>(Context.GetEntities());
+			Context.Defer().PushCommand<FMassCommandRemoveElements<FMassNavigationRelevantFragment, FMassNavigationRelevantParameters, FMassNavigationEvaluatedTag>>(Context.GetEntities());
+		});
+
+	// Remove the evaluated tag from non-navigation-relevant entities so they can be re-evaluated
+	NonRelevantEntityQuery.ForEachEntityChunk(ExecutionContext,
+		[](FMassExecutionContext& Context)
+		{
+			Context.Defer().PushCommand<FMassCommandRemoveElements<FMassNavigationEvaluatedTag>>(Context.GetEntities());
 		});
 }
diff --git a/Engine/Source/Runtime/Mass/MassEngine/Public/AI/MassEngineNavigationFragments.h b/Engine/Source/Runtime/Mass/MassEngine/Public/AI/MassEngineNavigationFragments.h
--- a/Engine/Source/Runtime/Mass/MassEngine/Public/AI/MassEngineNavigationFragments.h
+++ b/Engine/Source/Runtime/Mass/MassEngine/Public/AI/MassEngineNavigationFragments.h
@@ -21,6 +21,17 @@ struct FMassNavigationRelevantParameters : public FMassConstSharedFragment
 	bool bFillCollisionUnderneathForNavData = false;
 };
 
+/**
+ * Tag indicating that the entity has been evaluated for navigation relevance.
+ * Added to all entities after initial evaluation regardless of relevance outcome,
+ * so that non-relevant entities are not re-evaluated every frame.
+ */
+USTRUCT()
+struct FMassNavigationEvaluatedTag : public FMassTag
+{
+	GENERATED_BODY()
+};
+
 /**
  * Fragment holding the registration handle to the navigation element created from a Mass entity.
  * The fragment is added to indicate that a Mass entity is relevant to the AI navigation system.
diff --git a/Engine/Source/Runtime/Mass/MassEngine/Public/AI/MassEngineNavigationProcessors.h b/Engine/Source/Runtime/Mass/MassEngine/Public/AI/MassEngineNavigationProcessors.h
--- a/Engine/Source/Runtime/Mass/MassEngine/Public/AI/MassEngineNavigationProcessors.h
+++ b/Engine/Source/Runtime/Mass/MassEngine/Public/AI/MassEngineNavigationProcessors.h
@@ -68,4 +68,5 @@ protected:
 	virtual void Execute(FMassEntityManager& EntityManager, FMassExecutionContext& Context) override;
 
 	FMassEntityQuery EntityQuery;
+	FMassEntityQuery NonRelevantEntityQuery;
 };
\ No newline at end of file
